UnrealLightmass for Mac: restored building of Messages module and fixed a problem with editor waiting forever if the lightmass process crashed.

diff --git a/Engine/Source/Editor/SwarmInterface/Private/SwarmInterfaceLocal.cpp b/Engine/Source/Editor/SwarmInterface/Private/SwarmInterfaceLocal.cpp
--- a/Engine/Source/Editor/SwarmInterface/Private/SwarmInterfaceLocal.cpp
+++ b/Engine/Source/Editor/SwarmInterface/Private/SwarmInterfaceLocal.cpp
@@ -37,6 +37,7 @@ public:
 	virtual int32 CloseJob( void );
 	virtual int32 Log( TVerbosityLevel Verbosity, TLogColour TextColour, const TCHAR* Message );
 	virtual void SetJobGuid( const FGuid& JobGuid );
+	virtual bool IsJobProcessRunning( int32* OutStatus );
 
 private:
 
@@ -547,6 +548,16 @@ int32 FSwarmInterfaceLocalImpl::CloseJob( void )
 	return SWARM_SUCCESS;
 }
 
+bool FSwarmInterfaceLocalImpl::IsJobProcessRunning( int32* OutStatus )
+{
+	const bool bIsRunning = FPlatformProcess::IsProcRunning(LightmassProcHandle);
+	if (!bIsRunning && OutStatus)
+	{
+		*OutStatus = FPlatformProcess::GetProcReturnCode(LightmassProcHandle, OutStatus);
+	}
+	return bIsRunning;
+}
+
 int32 FSwarmInterfaceLocalImpl::Log( TVerbosityLevel Verbosity, TLogColour TextColour, const TCHAR* Message )
 {
 	if( Message == NULL )
diff --git a/Engine/Source/Editor/SwarmInterface/Public/SwarmInterface.h b/Engine/Source/Editor/SwarmInterface/Public/SwarmInterface.h
--- a/Engine/Source/Editor/SwarmInterface/Public/SwarmInterface.h
+++ b/Engine/Source/Editor/SwarmInterface/Public/SwarmInterface.h
@@ -568,6 +568,18 @@ public:
 	{
 	}
 
+	/**
+	 * Checks if the job's process is still running
+	 *
+	 * @param OutStatus Exit status if process is no longer running
+	 *
+	 * @return bool true if process is still running
+	 */
+	virtual bool IsJobProcessRunning( int32* OutStatus )
+	{
+		return true;
+	}
+
 protected:
 	FSwarmInterface( void )
 	{
diff --git a/Engine/Source/Editor/UnrealEd/Private/Lightmass/Lightmass.cpp b/Engine/Source/Editor/UnrealEd/Private/Lightmass/Lightmass.cpp
--- a/Engine/Source/Editor/UnrealEd/Private/Lightmass/Lightmass.cpp
+++ b/Engine/Source/Editor/UnrealEd/Private/Lightmass/Lightmass.cpp
@@ -2757,6 +2757,18 @@ bool FLightmassProcessor::Update()
 
 	ProcessAlertMessages();
 
+#if USE_LOCAL_SWARM_INTERFACE
+	int32 Status = 0;
+	const bool bIsLightmassRunning = Swarm.IsJobProcessRunning(&Status);
+	if (!bIsLightmassRunning)
+	{
+		bIsFinished = true;
+		bProcessingFailed == Status != 0;
+		bProcessingSuccessful = !bProcessingFailed;
+		bQuitReceived = true;
+	}
+#endif
+
 	return bIsFinished;
 }
 
diff --git a/Engine/Source/Programs/UnrealLightmass/UnrealLightmass.Build.cs b/Engine/Source/Programs/UnrealLightmass/UnrealLightmass.Build.cs
--- a/Engine/Source/Programs/UnrealLightmass/UnrealLightmass.Build.cs
+++ b/Engine/Source/Programs/UnrealLightmass/UnrealLightmass.Build.cs
@@ -25,6 +25,7 @@ public class UnrealLightmass : ModuleRules
 				new string[] {
 					"Networking",
 					"Sockets",
+					"Messaging",
 					"UdpMessaging",
 				}
 			);
diff --git a/Engine/Source/Runtime/Core/Public/GenericPlatform/GenericPlatformProcess.h b/Engine/Source/Runtime/Core/Public/GenericPlatform/GenericPlatformProcess.h
--- a/Engine/Source/Runtime/Core/Public/GenericPlatform/GenericPlatformProcess.h
+++ b/Engine/Source/Runtime/Core/Public/GenericPlatform/GenericPlatformProcess.h
@@ -106,7 +106,7 @@ public:
 	 */
 	FORCEINLINE bool Close()
 	{
-		return true;
+		return IsValid();
 	}
 
 protected:
