Non-Mac compile fix

diff --git a/Engine/Source/Editor/SwarmInterface/Private/SwarmInterfaceLocal.cpp b/Engine/Source/Editor/SwarmInterface/Private/SwarmInterfaceLocal.cpp
--- a/Engine/Source/Editor/SwarmInterface/Private/SwarmInterfaceLocal.cpp
+++ b/Engine/Source/Editor/SwarmInterface/Private/SwarmInterfaceLocal.cpp
@@ -550,12 +550,16 @@ int32 FSwarmInterfaceLocalImpl::CloseJob( void )
 
 bool FSwarmInterfaceLocalImpl::IsJobProcessRunning( int32* OutStatus )
 {
+#if USE_LOCAL_SWARM_INTERFACE
 	const bool bIsRunning = FPlatformProcess::IsProcRunning(LightmassProcHandle);
 	if (!bIsRunning && OutStatus)
 	{
 		*OutStatus = FPlatformProcess::GetProcReturnCode(LightmassProcHandle, OutStatus);
 	}
 	return bIsRunning;
+#else
+	return false;
+#endif
 }
 
 int32 FSwarmInterfaceLocalImpl::Log( TVerbosityLevel Verbosity, TLogColour TextColour, const TCHAR* Message )
