Fixed link errors in SwarmInteface on Mac

diff --git a/Engine/Source/Editor/SwarmInterface/SwarmInterface.Build.cs b/Engine/Source/Editor/SwarmInterface/SwarmInterface.Build.cs
--- a/Engine/Source/Editor/SwarmInterface/SwarmInterface.Build.cs
+++ b/Engine/Source/Editor/SwarmInterface/SwarmInterface.Build.cs
@@ -21,7 +21,7 @@ public class SwarmInterface : ModuleRules
 
 		if (Target.Platform == UnrealTargetPlatform.Mac || Target.Platform == UnrealTargetPlatform.Linux)
 		{
-			PrivateIncludePathModuleNames.Add("Messaging");
+			PrivateDependencyModuleNames.Add("Messaging");
 		}
 	}
 }
