Enabled UHT to digest GENERATED_BODY instead of GENERATED_UCLASS_BODY, GENERATED_USTRUCT_BODY, GENERATED_UINTERFACE_BODY or GENERATED_IINTERFACE_BODY, changed every occurence to the new syntax and fix

diff --git a/Engine/Source/Editor/Matinee/Classes/MatineeTrackColorPropHelper.h b/Engine/Source/Editor/Matinee/Classes/MatineeTrackColorPropHelper.h
--- a/Engine/Source/Editor/Matinee/Classes/MatineeTrackColorPropHelper.h
+++ b/Engine/Source/Editor/Matinee/Classes/MatineeTrackColorPropHelper.h
@@ -6,7 +6,9 @@
 UCLASS()
 class UMatineeTrackColorPropHelper : public UMatineeTrackVectorPropHelper
 {
-	GENERATED_UCLASS_BODY()
+	GENERATED_BODY()
+public:
+	UMatineeTrackColorPropHelper(const FObjectInitializer& ObjectInitializer = FObjectInitializer::Get());
 
 public:
 
diff --git a/Engine/Source/Editor/Matinee/Classes/MatineeTrackDirectorHelper.h b/Engine/Source/Editor/Matinee/Classes/MatineeTrackDirectorHelper.h
--- a/Engine/Source/Editor/Matinee/Classes/MatineeTrackDirectorHelper.h
+++ b/Engine/Source/Editor/Matinee/Classes/MatineeTrackDirectorHelper.h
@@ -8,7 +8,9 @@
 UCLASS()
 class UMatineeTrackDirectorHelper : public UInterpTrackHelper
 {
-	GENERATED_UCLASS_BODY()
+	GENERATED_BODY()
+public:
+	UMatineeTrackDirectorHelper(const FObjectInitializer& ObjectInitializer = FObjectInitializer::Get());
 
 	void OnAddKeyTextEntry(const FString& ChosenText, IMatineeBase* Matinee, UInterpTrack* Track);
 
diff --git a/Engine/Source/Editor/Matinee/Classes/MatineeTrackEventHelper.h b/Engine/Source/Editor/Matinee/Classes/MatineeTrackEventHelper.h
--- a/Engine/Source/Editor/Matinee/Classes/MatineeTrackEventHelper.h
+++ b/Engine/Source/Editor/Matinee/Classes/MatineeTrackEventHelper.h
@@ -8,7 +8,9 @@
 UCLASS()
 class UMatineeTrackEventHelper : public UInterpTrackHelper
 {
-	GENERATED_UCLASS_BODY()
+	GENERATED_BODY()
+public:
+	UMatineeTrackEventHelper(const FObjectInitializer& ObjectInitializer = FObjectInitializer::Get());
 
 	void OnAddKeyTextEntry(const FText& ChosenText, ETextCommit::Type CommitInfo, IMatineeBase* Matinee, UInterpTrack* Track);
 
diff --git a/Engine/Source/Editor/Matinee/Classes/MatineeTrackFloatPropHelper.h b/Engine/Source/Editor/Matinee/Classes/MatineeTrackFloatPropHelper.h
--- a/Engine/Source/Editor/Matinee/Classes/MatineeTrackFloatPropHelper.h
+++ b/Engine/Source/Editor/Matinee/Classes/MatineeTrackFloatPropHelper.h
@@ -10,7 +10,9 @@
 UCLASS()
 class UMatineeTrackFloatPropHelper : public UInterpTrackHelper
 {
-	GENERATED_UCLASS_BODY()
+	GENERATED_BODY()
+public:
+	UMatineeTrackFloatPropHelper(const FObjectInitializer& ObjectInitializer = FObjectInitializer::Get());
 
 	void OnCreateTrackTextEntry(const FString& ChosenText, TSharedRef<SWindow> Window, FString* OutputString);
 
diff --git a/Engine/Source/Editor/Matinee/Classes/MatineeTrackLinearColorPropHelper.h b/Engine/Source/Editor/Matinee/Classes/MatineeTrackLinearColorPropHelper.h
--- a/Engine/Source/Editor/Matinee/Classes/MatineeTrackLinearColorPropHelper.h
+++ b/Engine/Source/Editor/Matinee/Classes/MatineeTrackLinearColorPropHelper.h
@@ -8,7 +8,9 @@
 UCLASS()
 class UMatineeTrackLinearColorPropHelper : public UMatineeTrackVectorPropHelper
 {
-	GENERATED_UCLASS_BODY()
+	GENERATED_BODY()
+public:
+	UMatineeTrackLinearColorPropHelper(const FObjectInitializer& ObjectInitializer = FObjectInitializer::Get());
 
 public:
 
diff --git a/Engine/Source/Editor/Matinee/Classes/MatineeTrackParticleReplayHelper.h b/Engine/Source/Editor/Matinee/Classes/MatineeTrackParticleReplayHelper.h
--- a/Engine/Source/Editor/Matinee/Classes/MatineeTrackParticleReplayHelper.h
+++ b/Engine/Source/Editor/Matinee/Classes/MatineeTrackParticleReplayHelper.h
@@ -8,7 +8,9 @@
 UCLASS()
 class UMatineeTrackParticleReplayHelper : public UInterpTrackHelper
 {
-	GENERATED_UCLASS_BODY()
+	GENERATED_BODY()
+public:
+	UMatineeTrackParticleReplayHelper(const FObjectInitializer& ObjectInitializer = FObjectInitializer::Get());
 
 public:
 
diff --git a/Engine/Source/Editor/Matinee/Classes/MatineeTrackSoundHelper.h b/Engine/Source/Editor/Matinee/Classes/MatineeTrackSoundHelper.h
--- a/Engine/Source/Editor/Matinee/Classes/MatineeTrackSoundHelper.h
+++ b/Engine/Source/Editor/Matinee/Classes/MatineeTrackSoundHelper.h
@@ -8,7 +8,9 @@
 UCLASS()
 class UMatineeTrackSoundHelper : public UInterpTrackHelper
 {
-	GENERATED_UCLASS_BODY()
+	GENERATED_BODY()
+public:
+	UMatineeTrackSoundHelper(const FObjectInitializer& ObjectInitializer = FObjectInitializer::Get());
 
 public:
 
diff --git a/Engine/Source/Editor/Matinee/Classes/MatineeTrackToggleHelper.h b/Engine/Source/Editor/Matinee/Classes/MatineeTrackToggleHelper.h
--- a/Engine/Source/Editor/Matinee/Classes/MatineeTrackToggleHelper.h
+++ b/Engine/Source/Editor/Matinee/Classes/MatineeTrackToggleHelper.h
@@ -8,7 +8,9 @@
 UCLASS()
 class UMatineeTrackToggleHelper : public UInterpTrackHelper
 {
-	GENERATED_UCLASS_BODY()
+	GENERATED_BODY()
+public:
+	UMatineeTrackToggleHelper(const FObjectInitializer& ObjectInitializer = FObjectInitializer::Get());
 
 	void OnAddKeyTextEntry(const FString& ChosenText, IMatineeBase* Matinee, UInterpTrack* Track);
 
diff --git a/Engine/Source/Editor/Matinee/Classes/MatineeTrackVectorPropHelper.h b/Engine/Source/Editor/Matinee/Classes/MatineeTrackVectorPropHelper.h
--- a/Engine/Source/Editor/Matinee/Classes/MatineeTrackVectorPropHelper.h
+++ b/Engine/Source/Editor/Matinee/Classes/MatineeTrackVectorPropHelper.h
@@ -8,7 +8,9 @@
 UCLASS()
 class UMatineeTrackVectorPropHelper : public UInterpTrackHelper
 {
-	GENERATED_UCLASS_BODY()
+	GENERATED_BODY()
+public:
+	UMatineeTrackVectorPropHelper(const FObjectInitializer& ObjectInitializer = FObjectInitializer::Get());
 
 	/**
 	 * Pops up a dialog letting user choose between a set of properties, then checks to see if that property has been bound to yet.
diff --git a/Engine/Source/Editor/Matinee/Classes/MatineeTrackVisibilityHelper.h b/Engine/Source/Editor/Matinee/Classes/MatineeTrackVisibilityHelper.h
--- a/Engine/Source/Editor/Matinee/Classes/MatineeTrackVisibilityHelper.h
+++ b/Engine/Source/Editor/Matinee/Classes/MatineeTrackVisibilityHelper.h
@@ -7,7 +7,9 @@
 UCLASS()
 class UMatineeTrackVisibilityHelper : public UInterpTrackHelper
 {
-	GENERATED_UCLASS_BODY()
+	GENERATED_BODY()
+public:
+	UMatineeTrackVisibilityHelper(const FObjectInitializer& ObjectInitializer = FObjectInitializer::Get());
 
 	void OnAddKeyTextEntry(const FString& ChosenText, IMatineeBase* Matinee, UInterpTrack* Track);
 
diff --git a/Engine/Source/Editor/NewsFeed/Private/Implementation/NewsFeedSettings.h b/Engine/Source/Editor/NewsFeed/Private/Implementation/NewsFeedSettings.h
--- a/Engine/Source/Editor/NewsFeed/Private/Implementation/NewsFeedSettings.h
+++ b/Engine/Source/Editor/NewsFeed/Private/Implementation/NewsFeedSettings.h
@@ -29,7 +29,9 @@ UCLASS(config=EditorGameAgnostic)
 class UNewsFeedSettings
 	: public UObject
 {
-	GENERATED_UCLASS_BODY()
+	GENERATED_BODY()
+public:
+	UNewsFeedSettings(const FObjectInitializer& ObjectInitializer = FObjectInitializer::Get());
 
 public:
 
diff --git a/Engine/Source/Editor/OnlineBlueprintSupport/Classes/K2Node_InAppPurchase.h b/Engine/Source/Editor/OnlineBlueprintSupport/Classes/K2Node_InAppPurchase.h
--- a/Engine/Source/Editor/OnlineBlueprintSupport/Classes/K2Node_InAppPurchase.h
+++ b/Engine/Source/Editor/OnlineBlueprintSupport/Classes/K2Node_InAppPurchase.h
@@ -7,5 +7,7 @@
 UCLASS()
 class ONLINEBLUEPRINTSUPPORT_API UK2Node_InAppPurchase : public UK2Node_BaseAsyncTask
 {
-	GENERATED_UCLASS_BODY()
+	GENERATED_BODY()
+public:
+	UK2Node_InAppPurchase(const FObjectInitializer& ObjectInitializer = FObjectInitializer::Get());
 };
diff --git a/Engine/Source/Editor/OnlineBlueprintSupport/Classes/K2Node_InAppPurchaseQuery.h b/Engine/Source/Editor/OnlineBlueprintSupport/Classes/K2Node_InAppPurchaseQuery.h
--- a/Engine/Source/Editor/OnlineBlueprintSupport/Classes/K2Node_InAppPurchaseQuery.h
+++ b/Engine/Source/Editor/OnlineBlueprintSupport/Classes/K2Node_InAppPurchaseQuery.h
@@ -7,5 +7,7 @@
 UCLASS()
 class ONLINEBLUEPRINTSUPPORT_API UK2Node_InAppPurchaseQuery : public UK2Node_BaseAsyncTask
 {
-	GENERATED_UCLASS_BODY()
+	GENERATED_BODY()
+public:
+	UK2Node_InAppPurchaseQuery(const FObjectInitializer& ObjectInitializer = FObjectInitializer::Get());
 };
diff --git a/Engine/Source/Editor/OnlineBlueprintSupport/Classes/K2Node_LatentOnlineCall.h b/Engine/Source/Editor/OnlineBlueprintSupport/Classes/K2Node_LatentOnlineCall.h
--- a/Engine/Source/Editor/OnlineBlueprintSupport/Classes/K2Node_LatentOnlineCall.h
+++ b/Engine/Source/Editor/OnlineBlueprintSupport/Classes/K2Node_LatentOnlineCall.h
@@ -8,7 +8,9 @@
 UCLASS()
 class ONLINEBLUEPRINTSUPPORT_API UK2Node_LatentOnlineCall : public UK2Node_BaseAsyncTask
 {
-	GENERATED_UCLASS_BODY()
+	GENERATED_BODY()
+public:
+	UK2Node_LatentOnlineCall(const FObjectInitializer& ObjectInitializer = FObjectInitializer::Get());
 	
 	// UEdGraphNode interface
 	virtual void GetMenuEntries(FGraphContextMenuBuilder& ContextMenuBuilder) const override;
diff --git a/Engine/Source/Editor/OnlineBlueprintSupport/Classes/K2Node_LeaderboardFlush.h b/Engine/Source/Editor/OnlineBlueprintSupport/Classes/K2Node_LeaderboardFlush.h
--- a/Engine/Source/Editor/OnlineBlueprintSupport/Classes/K2Node_LeaderboardFlush.h
+++ b/Engine/Source/Editor/OnlineBlueprintSupport/Classes/K2Node_LeaderboardFlush.h
@@ -7,7 +7,9 @@
 UCLASS()
 class ONLINEBLUEPRINTSUPPORT_API UK2Node_LeaderboardFlush : public UK2Node_BaseAsyncTask
 {
-	GENERATED_UCLASS_BODY()
+	GENERATED_BODY()
+public:
+	UK2Node_LeaderboardFlush(const FObjectInitializer& ObjectInitializer = FObjectInitializer::Get());
 
 	// UEdGraphNode interface
 	virtual FText GetTooltipText() const override;
diff --git a/Engine/Source/Editor/OnlineBlueprintSupport/Classes/K2Node_LeaderboardQuery.h b/Engine/Source/Editor/OnlineBlueprintSupport/Classes/K2Node_LeaderboardQuery.h
--- a/Engine/Source/Editor/OnlineBlueprintSupport/Classes/K2Node_LeaderboardQuery.h
+++ b/Engine/Source/Editor/OnlineBlueprintSupport/Classes/K2Node_LeaderboardQuery.h
@@ -7,5 +7,7 @@
 UCLASS()
 class ONLINEBLUEPRINTSUPPORT_API UK2Node_LeaderboardQuery : public UK2Node_BaseAsyncTask
 {
-	GENERATED_UCLASS_BODY()
+	GENERATED_BODY()
+public:
+	UK2Node_LeaderboardQuery(const FObjectInitializer& ObjectInitializer = FObjectInitializer::Get());
 };
diff --git a/Engine/Source/Editor/PhAT/Classes/PhATEdSkeletalMeshComponent.h b/Engine/Source/Editor/PhAT/Classes/PhATEdSkeletalMeshComponent.h
--- a/Engine/Source/Editor/PhAT/Classes/PhATEdSkeletalMeshComponent.h
+++ b/Engine/Source/Editor/PhAT/Classes/PhATEdSkeletalMeshComponent.h
@@ -8,7 +8,9 @@
 UCLASS()
 class UPhATEdSkeletalMeshComponent : public USkeletalMeshComponent
 {
-	GENERATED_UCLASS_BODY()
+	GENERATED_BODY()
+public:
+	UPhATEdSkeletalMeshComponent(const FObjectInitializer& ObjectInitializer = FObjectInitializer::Get());
 
 	/** Data and methods shared across multiple classes */
 	class FPhATSharedData* SharedData;
diff --git a/Engine/Source/Editor/ReferenceViewer/Classes/EdGraphNode_Reference.h b/Engine/Source/Editor/ReferenceViewer/Classes/EdGraphNode_Reference.h
--- a/Engine/Source/Editor/ReferenceViewer/Classes/EdGraphNode_Reference.h
+++ b/Engine/Source/Editor/ReferenceViewer/Classes/EdGraphNode_Reference.h
@@ -7,7 +7,9 @@
 UCLASS()
 class UEdGraphNode_Reference : public UEdGraphNode
 {
-	GENERATED_UCLASS_BODY()
+	GENERATED_BODY()
+public:
+	UEdGraphNode_Reference(const FObjectInitializer& ObjectInitializer = FObjectInitializer::Get());
 
 	virtual void SetupReferenceNode(const FIntPoint& NodeLoc, const TArray<FName>& NewPackageNames, const FAssetData& InAssetData);
 	virtual void SetReferenceNodeCollapsed(const FIntPoint& NodeLoc, int32 InNumReferencesExceedingMax);
diff --git a/Engine/Source/Editor/ReferenceViewer/Classes/EdGraph_ReferenceViewer.h b/Engine/Source/Editor/ReferenceViewer/Classes/EdGraph_ReferenceViewer.h
--- a/Engine/Source/Editor/ReferenceViewer/Classes/EdGraph_ReferenceViewer.h
+++ b/Engine/Source/Editor/ReferenceViewer/Classes/EdGraph_ReferenceViewer.h
@@ -7,7 +7,9 @@
 UCLASS()
 class UEdGraph_ReferenceViewer : public UEdGraph
 {
-	GENERATED_UCLASS_BODY()
+	GENERATED_BODY()
+public:
+	UEdGraph_ReferenceViewer(const FObjectInitializer& ObjectInitializer = FObjectInitializer::Get());
 
 public:
 	// UObject implementation
diff --git a/Engine/Source/Editor/ReferenceViewer/Classes/ReferenceViewerSchema.h b/Engine/Source/Editor/ReferenceViewer/Classes/ReferenceViewerSchema.h
--- a/Engine/Source/Editor/ReferenceViewer/Classes/ReferenceViewerSchema.h
+++ b/Engine/Source/Editor/ReferenceViewer/Classes/ReferenceViewerSchema.h
@@ -6,7 +6,9 @@
 UCLASS()
 class UReferenceViewerSchema : public UEdGraphSchema
 {
-	GENERATED_UCLASS_BODY()
+	GENERATED_BODY()
+public:
+	UReferenceViewerSchema(const FObjectInitializer& ObjectInitializer = FObjectInitializer::Get());
 
 public:
 	// UEdGraphSchema interface
diff --git a/Engine/Source/Editor/SceneOutliner/Private/SceneOutlinerSettings.h b/Engine/Source/Editor/SceneOutliner/Private/SceneOutlinerSettings.h
--- a/Engine/Source/Editor/SceneOutliner/Private/SceneOutlinerSettings.h
+++ b/Engine/Source/Editor/SceneOutliner/Private/SceneOutlinerSettings.h
@@ -12,7 +12,9 @@ UCLASS(config=EditorUserSettings)
 class USceneOutlinerSettings
 	: public UObject
 {
-	GENERATED_UCLASS_BODY()
+	GENERATED_BODY()
+public:
+	USceneOutlinerSettings(const FObjectInitializer& ObjectInitializer = FObjectInitializer::Get());
 
 	/** True when the Scene Outliner is hiding temporary/run-time Actors */
 	UPROPERTY(config)
diff --git a/Engine/Source/Editor/Sequencer/Private/SequencerSettings.h b/Engine/Source/Editor/Sequencer/Private/SequencerSettings.h
--- a/Engine/Source/Editor/Sequencer/Private/SequencerSettings.h
+++ b/Engine/Source/Editor/Sequencer/Private/SequencerSettings.h
@@ -23,8 +23,10 @@ namespace ESequencerCurveVisibility
 UCLASS(config=EditorUserSettings)
 class USequencerSettings : public UObject
 {
+	GENERATED_BODY()
+public:
+	USequencerSettings(const FObjectInitializer& ObjectInitializer = FObjectInitializer::Get());
 public:
-	GENERATED_UCLASS_BODY()
 
 	DECLARE_MULTICAST_DELEGATE( FOnCurveVisibilityChanged );
 
@@ -140,4 +142,4 @@ protected:
 	TEnumAsByte<ESequencerCurveVisibility::Type> CurveVisibility;
 
 	FOnCurveVisibilityChanged OnCurveVisibilityChanged;
-};
\ No newline at end of file
+};
diff --git a/Engine/Source/Editor/StatsViewer/Classes/CookerStats.h b/Engine/Source/Editor/StatsViewer/Classes/CookerStats.h
--- a/Engine/Source/Editor/StatsViewer/Classes/CookerStats.h
+++ b/Engine/Source/Editor/StatsViewer/Classes/CookerStats.h
@@ -22,7 +22,9 @@ UCLASS(Transient, MinimalAPI, meta=(DisplayName="Cooker Stats", ObjectSetType="E
 class UCookerStats
 	: public UObject
 {
-	GENERATED_UCLASS_BODY()
+	GENERATED_BODY()
+public:
+	STATSVIEWER_API UCookerStats(const FObjectInitializer& ObjectInitializer = FObjectInitializer::Get());
 
 	/** The assets contained in the file. */
 	UPROPERTY(VisibleAnywhere, AssetRegistrySearchable, Category="Stats", meta=(DisplayName="Asset(s)", ColumnWidth="100"))
diff --git a/Engine/Source/Editor/StatsViewer/Classes/LightingBuildInfo.h b/Engine/Source/Editor/StatsViewer/Classes/LightingBuildInfo.h
--- a/Engine/Source/Editor/StatsViewer/Classes/LightingBuildInfo.h
+++ b/Engine/Source/Editor/StatsViewer/Classes/LightingBuildInfo.h
@@ -15,7 +15,9 @@ enum ELightingBuildInfoObjectSets
 UCLASS(Transient, MinimalAPI, meta=( DisplayName = "Lighting Build Info", ObjectSetType = "ELightingBuildInfoObjectSets" ) )
 class ULightingBuildInfo : public UObject
 {
-	GENERATED_UCLASS_BODY()
+	GENERATED_BODY()
+public:
+	STATSVIEWER_API ULightingBuildInfo(const FObjectInitializer& ObjectInitializer = FObjectInitializer::Get());
 
 	/** The actor and/or object that is related to this info. */
 	UPROPERTY(VisibleAnywhere, AssetRegistrySearchable, Category="Stats", meta=( ColumnWidth = "150" ) )
@@ -57,4 +59,4 @@ private:
 
 	/** Update internal strings */
 	void UpdateNames();
-};
\ No newline at end of file
+};
diff --git a/Engine/Source/Editor/StatsViewer/Classes/PrimitiveStats.h b/Engine/Source/Editor/StatsViewer/Classes/PrimitiveStats.h
--- a/Engine/Source/Editor/StatsViewer/Classes/PrimitiveStats.h
+++ b/Engine/Source/Editor/StatsViewer/Classes/PrimitiveStats.h
@@ -17,7 +17,9 @@ enum EPrimitiveObjectSets
 UCLASS(Transient, MinimalAPI, meta=( DisplayName = "Primitive Stats", ObjectSetType = "EPrimitiveObjectSets" ) )
 class UPrimitiveStats : public UObject
 {
-	GENERATED_UCLASS_BODY()
+	GENERATED_BODY()
+public:
+	STATSVIEWER_API UPrimitiveStats(const FObjectInitializer& ObjectInitializer = FObjectInitializer::Get());
 
 	/** Resource (e.g. UStaticMesh, USkeletalMesh, UModelComponent, UTerrainComponent, etc */
 	UPROPERTY(VisibleAnywhere, AssetRegistrySearchable, Category="Stats", meta=( DisplayName = "Object", ColumnWidth = "200" ) )
@@ -103,4 +105,4 @@ public:
 
 	/** Update internal strings */
 	void UpdateNames();
-};
\ No newline at end of file
+};
diff --git a/Engine/Source/Editor/StatsViewer/Classes/StaticMeshLightingInfo.h b/Engine/Source/Editor/StatsViewer/Classes/StaticMeshLightingInfo.h
--- a/Engine/Source/Editor/StatsViewer/Classes/StaticMeshLightingInfo.h
+++ b/Engine/Source/Editor/StatsViewer/Classes/StaticMeshLightingInfo.h
@@ -17,7 +17,9 @@ enum EStaticMeshLightingInfoObjectSets
 UCLASS(Transient, MinimalAPI, meta=( DisplayName = "Static Mesh Lighting Info", ObjectSetType = "EStaticMeshLightingInfoObjectSets" ) )
 class UStaticMeshLightingInfo : public UObject
 {
-	GENERATED_UCLASS_BODY()
+	GENERATED_BODY()
+public:
+	STATSVIEWER_API UStaticMeshLightingInfo(const FObjectInitializer& ObjectInitializer = FObjectInitializer::Get());
 
 	/** The actor that is related to this error/warning. */
 	UPROPERTY(VisibleAnywhere, AssetRegistrySearchable, Category="Stats", meta=( DisplayName = "Actor", ColumnWidth = "200" ) )
@@ -78,4 +80,4 @@ public:
 
 	/** Update internal strings */
 	void UpdateNames();
-};
\ No newline at end of file
+};
diff --git a/Engine/Source/Editor/StatsViewer/Classes/TextureStats.h b/Engine/Source/Editor/StatsViewer/Classes/TextureStats.h
--- a/Engine/Source/Editor/StatsViewer/Classes/TextureStats.h
+++ b/Engine/Source/Editor/StatsViewer/Classes/TextureStats.h
@@ -24,7 +24,9 @@ enum ETextureObjectSets
 UCLASS(Transient, MinimalAPI, meta=( DisplayName = "Texture Stats", ObjectSetType = "ETextureObjectSets" ) )
 class UTextureStats : public UObject
 {
-	GENERATED_UCLASS_BODY()
+	GENERATED_BODY()
+public:
+	STATSVIEWER_API UTextureStats(const FObjectInitializer& ObjectInitializer = FObjectInitializer::Get());
 
 	/** Texture - click to go to asset */
 	UPROPERTY(VisibleAnywhere, AssetRegistrySearchable, Category="Stats", meta=( DisplayName = "Name", ColumnWidth = "100" ) )
@@ -77,4 +79,4 @@ class UTextureStats : public UObject
 	/** Texture path without the name "package.[group.]" */
 	UPROPERTY(VisibleAnywhere, AssetRegistrySearchable, Category="Stats", meta=( ColumnWidth = "200" ) )
 	FString Path;
-};
\ No newline at end of file
+};
diff --git a/Engine/Source/Editor/SwarmInterface/Private/SwarmMessages.h b/Engine/Source/Editor/SwarmInterface/Private/SwarmMessages.h
--- a/Engine/Source/Editor/SwarmInterface/Private/SwarmMessages.h
+++ b/Engine/Source/Editor/SwarmInterface/Private/SwarmMessages.h
@@ -9,7 +9,7 @@
 USTRUCT()
 struct FSwarmPingMessage
 {
-	GENERATED_USTRUCT_BODY()
+	GENERATED_BODY()
 };
 
 template<>
@@ -25,7 +25,7 @@ struct TStructOpsTypeTraits<FSwarmPingMessage> : public TStructOpsTypeTraitsBase
 USTRUCT()
 struct FSwarmPongMessage
 {
-	GENERATED_USTRUCT_BODY()
+	GENERATED_BODY()
 
 	UPROPERTY()
 	bool			bIsEditor;
@@ -55,7 +55,7 @@ struct TStructOpsTypeTraits<FSwarmPongMessage> : public TStructOpsTypeTraitsBase
 USTRUCT()
 struct FSwarmInfoMessage
 {
-	GENERATED_USTRUCT_BODY()
+	GENERATED_BODY()
 
 	UPROPERTY()
 	FString TextMessage;
@@ -81,7 +81,7 @@ struct TStructOpsTypeTraits<FSwarmInfoMessage> : public TStructOpsTypeTraitsBase
 USTRUCT()
 struct FSwarmAlertMessage
 {
-	GENERATED_USTRUCT_BODY()
+	GENERATED_BODY()
 
 	/** The Job Guid */
 	UPROPERTY()
@@ -124,7 +124,7 @@ struct TStructOpsTypeTraits<FSwarmAlertMessage> : public TStructOpsTypeTraitsBas
 USTRUCT()
 struct FSwarmTimingMessage
 {
-	GENERATED_USTRUCT_BODY()
+	GENERATED_BODY()
 
 	/** State that the distributed job is transitioning to */
 	UPROPERTY()
@@ -155,7 +155,7 @@ struct TStructOpsTypeTraits<FSwarmTimingMessage> : public TStructOpsTypeTraitsBa
 USTRUCT()
 struct FSwarmTaskRequestMessage
 {
-	GENERATED_USTRUCT_BODY()
+	GENERATED_BODY()
 
 	FSwarmTaskRequestMessage() {}
 };
@@ -173,7 +173,7 @@ struct TStructOpsTypeTraits<FSwarmTaskRequestMessage> : public TStructOpsTypeTra
 USTRUCT()
 struct FSwarmTaskRequestReleaseMessage
 {
-	GENERATED_USTRUCT_BODY()
+	GENERATED_BODY()
 
 	FSwarmTaskRequestReleaseMessage() {}
 };
@@ -191,7 +191,7 @@ struct TStructOpsTypeTraits<FSwarmTaskRequestReleaseMessage> : public TStructOps
 USTRUCT()
 struct FSwarmTaskRequestReservationMessage
 {
-	GENERATED_USTRUCT_BODY()
+	GENERATED_BODY()
 
 	FSwarmTaskRequestReservationMessage() {}
 };
@@ -209,7 +209,7 @@ struct TStructOpsTypeTraits<FSwarmTaskRequestReservationMessage> : public TStruc
 USTRUCT()
 struct FSwarmTaskRequestSpecificationMessage
 {
-	GENERATED_USTRUCT_BODY()
+	GENERATED_BODY()
 
 	/** The GUID used for identifying the Task being referred to */
 	UPROPERTY()
@@ -250,7 +250,7 @@ struct TStructOpsTypeTraits<FSwarmTaskRequestSpecificationMessage> : public TStr
 USTRUCT()
 struct FSwarmJobStateMessage
 {
-	GENERATED_USTRUCT_BODY()
+	GENERATED_BODY()
 
 	/** The Job GUID used for identifying the Job */
 	UPROPERTY()
@@ -291,7 +291,7 @@ struct TStructOpsTypeTraits<FSwarmJobStateMessage> : public TStructOpsTypeTraits
 USTRUCT()
 struct FSwarmTaskStateMessage
 {
-	GENERATED_USTRUCT_BODY()
+	GENERATED_BODY()
 
 	/** The Task GUID used for identifying the Task */
 	UPROPERTY()
@@ -332,7 +332,7 @@ struct TStructOpsTypeTraits<FSwarmTaskStateMessage> : public TStructOpsTypeTrait
 USTRUCT()
 struct FSwarmQuitMessage
 {
-	GENERATED_USTRUCT_BODY()
+	GENERATED_BODY()
 
 	FSwarmQuitMessage() {}
 };
diff --git a/Engine/Source/Editor/TextureEditor/Classes/TextureEditorSettings.h b/Engine/Source/Editor/TextureEditor/Classes/TextureEditorSettings.h
--- a/Engine/Source/Editor/TextureEditor/Classes/TextureEditorSettings.h
+++ b/Engine/Source/Editor/TextureEditor/Classes/TextureEditorSettings.h
@@ -25,7 +25,9 @@ UCLASS(config=EditorUserSettings)
 class TEXTUREEDITOR_API UTextureEditorSettings
 	: public UObject
 {
-	GENERATED_UCLASS_BODY()
+	GENERATED_BODY()
+public:
+	UTextureEditorSettings(const FObjectInitializer& ObjectInitializer = FObjectInitializer::Get());
 
 public:
 
diff --git a/Engine/Source/Editor/TranslationEditor/Private/InternationalizationExportSettings.h b/Engine/Source/Editor/TranslationEditor/Private/InternationalizationExportSettings.h
--- a/Engine/Source/Editor/TranslationEditor/Private/InternationalizationExportSettings.h
+++ b/Engine/Source/Editor/TranslationEditor/Private/InternationalizationExportSettings.h
@@ -6,7 +6,9 @@
 UCLASS(hidecategories = Object, MinimalAPI, config=InternationalizationExport)
 class UInternationalizationExportSettings : public UObject
 {
-	GENERATED_UCLASS_BODY()
+	GENERATED_BODY()
+public:
+	TRANSLATIONEDITOR_API UInternationalizationExportSettings(const FObjectInitializer& ObjectInitializer = FObjectInitializer::Get());
 
 public:
 
@@ -49,4 +51,4 @@ public:
 	/** Whether or not to use culture path */
 	UPROPERTY(Category = GatherTextStep, EditAnywhere, config)
 	bool bUseCultureDirectory;
-};
\ No newline at end of file
+};
diff --git a/Engine/Source/Editor/TranslationEditor/Private/TranslationUnit.h b/Engine/Source/Editor/TranslationEditor/Private/TranslationUnit.h
--- a/Engine/Source/Editor/TranslationEditor/Private/TranslationUnit.h
+++ b/Engine/Source/Editor/TranslationEditor/Private/TranslationUnit.h
@@ -6,7 +6,7 @@
 USTRUCT()
 struct FTranslationChange
 {
-	GENERATED_USTRUCT_BODY()
+	GENERATED_BODY()
 
 public:
 
@@ -30,7 +30,7 @@ public:
 USTRUCT()
 struct FTranslationContextInfo
 {
-	GENERATED_USTRUCT_BODY()
+	GENERATED_BODY()
 
 public:
 
@@ -50,7 +50,9 @@ public:
 UCLASS(hidecategories = Object, MinimalAPI)
 class UTranslationUnit : public UObject
 {
-	GENERATED_UCLASS_BODY()
+	GENERATED_BODY()
+public:
+	TRANSLATIONEDITOR_API UTranslationUnit(const FObjectInitializer& ObjectInitializer = FObjectInitializer::Get());
 
 public:
 
@@ -102,4 +104,4 @@ private:
 
 	// Holds an event delegate that is executed when a property has changed.
 	FTranslationUnitPropertyChangedEvent TranslationUnitPropertyChangedEvent;
-};
\ No newline at end of file
+};
diff --git a/Engine/Source/Editor/UMGEditor/Classes/WidgetBlueprintFactory.h b/Engine/Source/Editor/UMGEditor/Classes/WidgetBlueprintFactory.h
--- a/Engine/Source/Editor/UMGEditor/Classes/WidgetBlueprintFactory.h
+++ b/Engine/Source/Editor/UMGEditor/Classes/WidgetBlueprintFactory.h
@@ -7,7 +7,9 @@
 UCLASS(HideCategories=Object, MinimalAPI)
 class UWidgetBlueprintFactory : public UFactory
 {
-	GENERATED_UCLASS_BODY()
+	GENERATED_BODY()
+public:
+	UMGEDITOR_API UWidgetBlueprintFactory(const FObjectInitializer& ObjectInitializer = FObjectInitializer::Get());
 
 	// The type of blueprint that will be created
 	UPROPERTY(EditAnywhere, Category=WidgetBlueprintFactory)
diff --git a/Engine/Source/Editor/UMGEditor/Classes/WidgetGraphSchema.h b/Engine/Source/Editor/UMGEditor/Classes/WidgetGraphSchema.h
--- a/Engine/Source/Editor/UMGEditor/Classes/WidgetGraphSchema.h
+++ b/Engine/Source/Editor/UMGEditor/Classes/WidgetGraphSchema.h
@@ -8,7 +8,9 @@
 UCLASS(MinimalAPI)
 class UWidgetGraphSchema : public UEdGraphSchema_K2
 {
-	GENERATED_UCLASS_BODY()
+	GENERATED_BODY()
+public:
+	UMGEDITOR_API UWidgetGraphSchema(const FObjectInitializer& ObjectInitializer = FObjectInitializer::Get());
 
 	virtual bool ShouldAlwaysPurgeOnModification() const override { return true; }
 
@@ -31,4 +33,4 @@ class UWidgetGraphSchema : public UEdGraphSchema_K2
 	FName DefaultEvaluationHandlerName;
 	
 	// TODO UMG - Add more things specific to the needs of the widget blueprint graph.
-};
\ No newline at end of file
+};
diff --git a/Engine/Source/Editor/UMGEditor/Private/Nodes/K2Node_CreateDragDropOperation.h b/Engine/Source/Editor/UMGEditor/Private/Nodes/K2Node_CreateDragDropOperation.h
--- a/Engine/Source/Editor/UMGEditor/Private/Nodes/K2Node_CreateDragDropOperation.h
+++ b/Engine/Source/Editor/UMGEditor/Private/Nodes/K2Node_CreateDragDropOperation.h
@@ -9,7 +9,9 @@
 UCLASS()
 class UMGEDITOR_API UK2Node_CreateDragDropOperation : public UK2Node_ConstructObjectFromClass
 {
-	GENERATED_UCLASS_BODY()
+	GENERATED_BODY()
+public:
+	UK2Node_CreateDragDropOperation(const FObjectInitializer& ObjectInitializer = FObjectInitializer::Get());
 
 	// Begin UEdGraphNode interface.
 	virtual void AllocateDefaultPins() override;
diff --git a/Engine/Source/Editor/UMGEditor/Private/Nodes/K2Node_CreateWidget.h b/Engine/Source/Editor/UMGEditor/Private/Nodes/K2Node_CreateWidget.h
--- a/Engine/Source/Editor/UMGEditor/Private/Nodes/K2Node_CreateWidget.h
+++ b/Engine/Source/Editor/UMGEditor/Private/Nodes/K2Node_CreateWidget.h
@@ -9,7 +9,9 @@
 UCLASS()
 class UMGEDITOR_API UK2Node_CreateWidget : public UK2Node_ConstructObjectFromClass
 {
-	GENERATED_UCLASS_BODY()
+	GENERATED_BODY()
+public:
+	UK2Node_CreateWidget(const FObjectInitializer& ObjectInitializer = FObjectInitializer::Get());
 
 	// Begin UEdGraphNode interface.
 	virtual void AllocateDefaultPins() override;
diff --git a/Engine/Source/Editor/UMGEditor/Public/Settings/WidgetDesignerSettings.h b/Engine/Source/Editor/UMGEditor/Public/Settings/WidgetDesignerSettings.h
--- a/Engine/Source/Editor/UMGEditor/Public/Settings/WidgetDesignerSettings.h
+++ b/Engine/Source/Editor/UMGEditor/Public/Settings/WidgetDesignerSettings.h
@@ -11,7 +11,9 @@ UCLASS(config=EditorUserSettings)
 class UMGEDITOR_API UWidgetDesignerSettings
 	: public UObject
 {
-	GENERATED_UCLASS_BODY()
+	GENERATED_BODY()
+public:
+	UWidgetDesignerSettings(const FObjectInitializer& ObjectInitializer = FObjectInitializer::Get());
 
 	/** If enabled, actor positions will snap to the grid. */
 	UPROPERTY(EditAnywhere, config, Category=GridSnapping, meta=(DisplayName = "Enable Grid Snapping"))
diff --git a/Engine/Source/Editor/UMGEditor/Public/WidgetBlueprint.h b/Engine/Source/Editor/UMGEditor/Public/WidgetBlueprint.h
--- a/Engine/Source/Editor/UMGEditor/Public/WidgetBlueprint.h
+++ b/Engine/Source/Editor/UMGEditor/Public/WidgetBlueprint.h
@@ -15,7 +15,7 @@ class UMovieScene;
 USTRUCT()
 struct UMGEDITOR_API FEditorPropertyPathSegment
 {
-	GENERATED_USTRUCT_BODY()
+	GENERATED_BODY()
 
 public:
 	FEditorPropertyPathSegment();
@@ -59,7 +59,7 @@ private:
 USTRUCT()
 struct UMGEDITOR_API FEditorPropertyPath
 {
-	GENERATED_USTRUCT_BODY()
+	GENERATED_BODY()
 
 public:
 
@@ -96,7 +96,7 @@ public:
 USTRUCT()
 struct UMGEDITOR_API FDelegateEditorBinding
 {
-	GENERATED_USTRUCT_BODY()
+	GENERATED_BODY()
 
 	/** The member widget the binding is on, must be a direct variable of the UUserWidget. */
 	UPROPERTY()
@@ -142,7 +142,7 @@ struct UMGEDITOR_API FDelegateEditorBinding
 USTRUCT()
 struct FWidgetAnimation_DEPRECATED
 {
-	GENERATED_USTRUCT_BODY()
+	GENERATED_BODY()
 
 	UPROPERTY()
 	UMovieScene* MovieScene;
@@ -170,7 +170,9 @@ struct TStructOpsTypeTraits<FWidgetAnimation_DEPRECATED> : public TStructOpsType
 UCLASS(BlueprintType)
 class UMGEDITOR_API UWidgetBlueprint : public UBlueprint
 {
-	GENERATED_UCLASS_BODY()
+	GENERATED_BODY()
+public:
+	UWidgetBlueprint(const FObjectInitializer& ObjectInitializer = FObjectInitializer::Get());
 
 public:
 
@@ -222,4 +224,4 @@ public:
 	// End of UBlueprint interface
 
 	static bool ValidateGeneratedClass(const UClass* InClass);
-};
\ No newline at end of file
+};
diff --git a/Engine/Source/Editor/UnrealEd/Classes/ActorFactories/ActorFactory.h b/Engine/Source/Editor/UnrealEd/Classes/ActorFactories/ActorFactory.h
--- a/Engine/Source/Editor/UnrealEd/Classes/ActorFactories/ActorFactory.h
+++ b/Engine/Source/Editor/UnrealEd/Classes/ActorFactories/ActorFactory.h
@@ -11,7 +11,9 @@ class ULevel;
 UCLASS(collapsecategories, hidecategories=Object, editinlinenew, config=Editor, abstract, transient)
 class UNREALED_API UActorFactory : public UObject
 {
-	GENERATED_UCLASS_BODY()
+	GENERATED_BODY()
+public:
+	UActorFactory(const FObjectInitializer& ObjectInitializer = FObjectInitializer::Get());
 
 	/** Name used as basis for 'New Actor' menu. */
 	UPROPERTY()
diff --git a/Engine/Source/Editor/UnrealEd/Classes/ActorFactories/ActorFactoryAmbientSound.h b/Engine/Source/Editor/UnrealEd/Classes/ActorFactories/ActorFactoryAmbientSound.h
--- a/Engine/Source/Editor/UnrealEd/Classes/ActorFactories/ActorFactoryAmbientSound.h
+++ b/Engine/Source/Editor/UnrealEd/Classes/ActorFactories/ActorFactoryAmbientSound.h
@@ -12,7 +12,9 @@
 UCLASS(MinimalAPI, config=Editor, collapsecategories, hidecategories=Object)
 class UActorFactoryAmbientSound : public UActorFactory
 {
-	GENERATED_UCLASS_BODY()
+	GENERATED_BODY()
+public:
+	UNREALED_API UActorFactoryAmbientSound(const FObjectInitializer& ObjectInitializer = FObjectInitializer::Get());
 
 	// Begin UActorFactory Interface
 	virtual void PostSpawnActor( UObject* Asset, AActor* NewActor ) override;
diff --git a/Engine/Source/Editor/UnrealEd/Classes/ActorFactories/ActorFactoryAnimationAsset.h b/Engine/Source/Editor/UnrealEd/Classes/ActorFactories/ActorFactoryAnimationAsset.h
--- a/Engine/Source/Editor/UnrealEd/Classes/ActorFactories/ActorFactoryAnimationAsset.h
+++ b/Engine/Source/Editor/UnrealEd/Classes/ActorFactories/ActorFactoryAnimationAsset.h
@@ -10,7 +10,9 @@ class USkeletalMesh;
 UCLASS(MinimalAPI, config=Editor, hidecategories=Object)
 class UActorFactoryAnimationAsset : public UActorFactorySkeletalMesh
 {
-	GENERATED_UCLASS_BODY()
+	GENERATED_BODY()
+public:
+	UNREALED_API UActorFactoryAnimationAsset(const FObjectInitializer& ObjectInitializer = FObjectInitializer::Get());
 
 protected:
 	// Begin UActorFactory Interface
diff --git a/Engine/Source/Editor/UnrealEd/Classes/ActorFactories/ActorFactoryAtmosphericFog.h b/Engine/Source/Editor/UnrealEd/Classes/ActorFactories/ActorFactoryAtmosphericFog.h
--- a/Engine/Source/Editor/UnrealEd/Classes/ActorFactories/ActorFactoryAtmosphericFog.h
+++ b/Engine/Source/Editor/UnrealEd/Classes/ActorFactories/ActorFactoryAtmosphericFog.h
@@ -7,7 +7,9 @@
 UCLASS(MinimalAPI, config=Editor, collapsecategories, hidecategories=Object)
 class UActorFactoryAtmosphericFog : public UActorFactory
 {
-	GENERATED_UCLASS_BODY()
+	GENERATED_BODY()
+public:
+	UNREALED_API UActorFactoryAtmosphericFog(const FObjectInitializer& ObjectInitializer = FObjectInitializer::Get());
 };
 
 
diff --git a/Engine/Source/Editor/UnrealEd/Classes/ActorFactories/ActorFactoryBasicShape.h b/Engine/Source/Editor/UnrealEd/Classes/ActorFactories/ActorFactoryBasicShape.h
--- a/Engine/Source/Editor/UnrealEd/Classes/ActorFactories/ActorFactoryBasicShape.h
+++ b/Engine/Source/Editor/UnrealEd/Classes/ActorFactories/ActorFactoryBasicShape.h
@@ -7,7 +7,9 @@
 UCLASS(MinimalAPI,config=Editor)
 class UActorFactoryBasicShape : public UActorFactoryStaticMesh
 {
-	GENERATED_UCLASS_BODY()
+	GENERATED_BODY()
+public:
+	UNREALED_API UActorFactoryBasicShape(const FObjectInitializer& ObjectInitializer = FObjectInitializer::Get());
 
 	virtual bool CanCreateActorFrom( const FAssetData& AssetData, FText& OutErrorMsg ) override;
 	virtual void PostSpawnActor(UObject* Asset, AActor* NewActor) override;
diff --git a/Engine/Source/Editor/UnrealEd/Classes/ActorFactories/ActorFactoryBlueprint.h b/Engine/Source/Editor/UnrealEd/Classes/ActorFactories/ActorFactoryBlueprint.h
--- a/Engine/Source/Editor/UnrealEd/Classes/ActorFactories/ActorFactoryBlueprint.h
+++ b/Engine/Source/Editor/UnrealEd/Classes/ActorFactories/ActorFactoryBlueprint.h
@@ -8,7 +8,9 @@
 UCLASS(MinimalAPI, config=Editor, collapsecategories, hidecategories=Object)
 class UActorFactoryBlueprint : public UActorFactory
 {
-	GENERATED_UCLASS_BODY()
+	GENERATED_BODY()
+public:
+	UNREALED_API UActorFactoryBlueprint(const FObjectInitializer& ObjectInitializer = FObjectInitializer::Get());
 
 	// Begin UActorFactory Interface
 	virtual bool PreSpawnActor( UObject* Asset, FTransform& InOutLocation ) override;
diff --git a/Engine/Source/Editor/UnrealEd/Classes/ActorFactories/ActorFactoryBoxReflectionCapture.h b/Engine/Source/Editor/UnrealEd/Classes/ActorFactories/ActorFactoryBoxReflectionCapture.h
--- a/Engine/Source/Editor/UnrealEd/Classes/ActorFactories/ActorFactoryBoxReflectionCapture.h
+++ b/Engine/Source/Editor/UnrealEd/Classes/ActorFactories/ActorFactoryBoxReflectionCapture.h
@@ -7,7 +7,9 @@
 UCLASS(MinimalAPI, config=Editor, collapsecategories, hidecategories=Object)
 class UActorFactoryBoxReflectionCapture : public UActorFactory
 {
-	GENERATED_UCLASS_BODY()
+	GENERATED_BODY()
+public:
+	UNREALED_API UActorFactoryBoxReflectionCapture(const FObjectInitializer& ObjectInitializer = FObjectInitializer::Get());
 };
 
 
diff --git a/Engine/Source/Editor/UnrealEd/Classes/ActorFactories/ActorFactoryBoxVolume.h b/Engine/Source/Editor/UnrealEd/Classes/ActorFactories/ActorFactoryBoxVolume.h
--- a/Engine/Source/Editor/UnrealEd/Classes/ActorFactories/ActorFactoryBoxVolume.h
+++ b/Engine/Source/Editor/UnrealEd/Classes/ActorFactories/ActorFactoryBoxVolume.h
@@ -8,7 +8,9 @@
 UCLASS(MinimalAPI, config=Editor, collapsecategories, hidecategories=Object)
 class UActorFactoryBoxVolume : public UActorFactory
 {
-	GENERATED_UCLASS_BODY()
+	GENERATED_BODY()
+public:
+	UNREALED_API UActorFactoryBoxVolume(const FObjectInitializer& ObjectInitializer = FObjectInitializer::Get());
 
 	// Begin UActorFactory Interface
 	virtual bool CanCreateActorFrom( const FAssetData& AssetData, FText& OutErrorMsg ) override;
diff --git a/Engine/Source/Editor/UnrealEd/Classes/ActorFactories/ActorFactoryCameraActor.h b/Engine/Source/Editor/UnrealEd/Classes/ActorFactories/ActorFactoryCameraActor.h
--- a/Engine/Source/Editor/UnrealEd/Classes/ActorFactories/ActorFactoryCameraActor.h
+++ b/Engine/Source/Editor/UnrealEd/Classes/ActorFactories/ActorFactoryCameraActor.h
@@ -7,5 +7,7 @@
 UCLASS(MinimalAPI, config=Editor, collapsecategories, hidecategories=Object)
 class UActorFactoryCameraActor : public UActorFactory
 {
-	GENERATED_UCLASS_BODY()
+	GENERATED_BODY()
+public:
+	UNREALED_API UActorFactoryCameraActor(const FObjectInitializer& ObjectInitializer = FObjectInitializer::Get());
 };
diff --git a/Engine/Source/Editor/UnrealEd/Classes/ActorFactories/ActorFactoryCharacter.h b/Engine/Source/Editor/UnrealEd/Classes/ActorFactories/ActorFactoryCharacter.h
--- a/Engine/Source/Editor/UnrealEd/Classes/ActorFactories/ActorFactoryCharacter.h
+++ b/Engine/Source/Editor/UnrealEd/Classes/ActorFactories/ActorFactoryCharacter.h
@@ -7,7 +7,9 @@
 UCLASS(MinimalAPI,config=Editor)
 class UActorFactoryCharacter : public UActorFactory
 {
-	GENERATED_UCLASS_BODY()
+	GENERATED_BODY()
+public:
+	UNREALED_API UActorFactoryCharacter(const FObjectInitializer& ObjectInitializer = FObjectInitializer::Get());
 
 public:
 	virtual bool CanCreateActorFrom( const FAssetData& AssetData, FText& OutErrorMsg ) override;
diff --git a/Engine/Source/Editor/UnrealEd/Classes/ActorFactories/ActorFactoryClass.h b/Engine/Source/Editor/UnrealEd/Classes/ActorFactories/ActorFactoryClass.h
--- a/Engine/Source/Editor/UnrealEd/Classes/ActorFactories/ActorFactoryClass.h
+++ b/Engine/Source/Editor/UnrealEd/Classes/ActorFactories/ActorFactoryClass.h
@@ -8,7 +8,9 @@
 UCLASS(MinimalAPI, config=Editor, collapsecategories, hidecategories=Object)
 class UActorFactoryClass : public UActorFactory
 {
-	GENERATED_UCLASS_BODY()
+	GENERATED_BODY()
+public:
+	UNREALED_API UActorFactoryClass(const FObjectInitializer& ObjectInitializer = FObjectInitializer::Get());
 
 	// Begin UActorFactory Interface
 	virtual bool PreSpawnActor( UObject* Asset, FTransform& InOutLocation ) override;
diff --git a/Engine/Source/Editor/UnrealEd/Classes/ActorFactories/ActorFactoryCylinderVolume.h b/Engine/Source/Editor/UnrealEd/Classes/ActorFactories/ActorFactoryCylinderVolume.h
--- a/Engine/Source/Editor/UnrealEd/Classes/ActorFactories/ActorFactoryCylinderVolume.h
+++ b/Engine/Source/Editor/UnrealEd/Classes/ActorFactories/ActorFactoryCylinderVolume.h
@@ -8,7 +8,9 @@
 UCLASS(MinimalAPI, config=Editor, collapsecategories, hidecategories=Object)
 class UActorFactoryCylinderVolume : public UActorFactory
 {
-	GENERATED_UCLASS_BODY()
+	GENERATED_BODY()
+public:
+	UNREALED_API UActorFactoryCylinderVolume(const FObjectInitializer& ObjectInitializer = FObjectInitializer::Get());
 
 	// Begin UActorFactory Interface
 	virtual bool CanCreateActorFrom( const FAssetData& AssetData, FText& OutErrorMsg ) override;
diff --git a/Engine/Source/Editor/UnrealEd/Classes/ActorFactories/ActorFactoryDeferredDecal.h b/Engine/Source/Editor/UnrealEd/Classes/ActorFactories/ActorFactoryDeferredDecal.h
--- a/Engine/Source/Editor/UnrealEd/Classes/ActorFactories/ActorFactoryDeferredDecal.h
+++ b/Engine/Source/Editor/UnrealEd/Classes/ActorFactories/ActorFactoryDeferredDecal.h
@@ -8,8 +8,10 @@
 UCLASS(MinimalAPI, config=Editor)
 class UActorFactoryDeferredDecal : public UActorFactory
 {
+	GENERATED_BODY()
+public:
+	UNREALED_API UActorFactoryDeferredDecal(const FObjectInitializer& ObjectInitializer = FObjectInitializer::Get());
 public:
-	GENERATED_UCLASS_BODY()
 
 protected:
 	// Begin UActorFactory Interface
diff --git a/Engine/Source/Editor/UnrealEd/Classes/ActorFactories/ActorFactoryDestructible.h b/Engine/Source/Editor/UnrealEd/Classes/ActorFactories/ActorFactoryDestructible.h
--- a/Engine/Source/Editor/UnrealEd/Classes/ActorFactories/ActorFactoryDestructible.h
+++ b/Engine/Source/Editor/UnrealEd/Classes/ActorFactories/ActorFactoryDestructible.h
@@ -8,7 +8,9 @@
 UCLASS(MinimalAPI, config=Editor)
 class UActorFactoryDestructible : public UActorFactory
 {
-	GENERATED_UCLASS_BODY()
+	GENERATED_BODY()
+public:
+	UNREALED_API UActorFactoryDestructible(const FObjectInitializer& ObjectInitializer = FObjectInitializer::Get());
 
 	// Begin UActorFactory Interface
 	virtual void PostSpawnActor( UObject* Asset, AActor* NewActor ) override;
diff --git a/Engine/Source/Editor/UnrealEd/Classes/ActorFactories/ActorFactoryDirectionalLight.h b/Engine/Source/Editor/UnrealEd/Classes/ActorFactories/ActorFactoryDirectionalLight.h
--- a/Engine/Source/Editor/UnrealEd/Classes/ActorFactories/ActorFactoryDirectionalLight.h
+++ b/Engine/Source/Editor/UnrealEd/Classes/ActorFactories/ActorFactoryDirectionalLight.h
@@ -7,7 +7,9 @@
 UCLASS(MinimalAPI, config=Editor, collapsecategories, hidecategories=Object)
 class UActorFactoryDirectionalLight : public UActorFactory
 {
-	GENERATED_UCLASS_BODY()
+	GENERATED_BODY()
+public:
+	UNREALED_API UActorFactoryDirectionalLight(const FObjectInitializer& ObjectInitializer = FObjectInitializer::Get());
 
 };
 
diff --git a/Engine/Source/Editor/UnrealEd/Classes/ActorFactories/ActorFactoryEmitter.h b/Engine/Source/Editor/UnrealEd/Classes/ActorFactories/ActorFactoryEmitter.h
--- a/Engine/Source/Editor/UnrealEd/Classes/ActorFactories/ActorFactoryEmitter.h
+++ b/Engine/Source/Editor/UnrealEd/Classes/ActorFactories/ActorFactoryEmitter.h
@@ -8,7 +8,9 @@
 UCLASS(MinimalAPI, config=Editor, collapsecategories, hidecategories=Object)
 class UActorFactoryEmitter : public UActorFactory
 {
-	GENERATED_UCLASS_BODY()
+	GENERATED_BODY()
+public:
+	UNREALED_API UActorFactoryEmitter(const FObjectInitializer& ObjectInitializer = FObjectInitializer::Get());
 
 	// Begin UActorFactory Interface
 	virtual void PostSpawnActor( UObject* Asset, AActor* NewActor ) override;
diff --git a/Engine/Source/Editor/UnrealEd/Classes/ActorFactories/ActorFactoryEmptyActor.h b/Engine/Source/Editor/UnrealEd/Classes/ActorFactories/ActorFactoryEmptyActor.h
--- a/Engine/Source/Editor/UnrealEd/Classes/ActorFactories/ActorFactoryEmptyActor.h
+++ b/Engine/Source/Editor/UnrealEd/Classes/ActorFactories/ActorFactoryEmptyActor.h
@@ -7,7 +7,9 @@
 UCLASS(MinimalAPI,config=Editor)
 class UActorFactoryEmptyActor : public UActorFactory
 {
-	GENERATED_UCLASS_BODY()
+	GENERATED_BODY()
+public:
+	UNREALED_API UActorFactoryEmptyActor(const FObjectInitializer& ObjectInitializer = FObjectInitializer::Get());
 
 public:
 	virtual bool CanCreateActorFrom( const FAssetData& AssetData, FText& OutErrorMsg ) override;
diff --git a/Engine/Source/Editor/UnrealEd/Classes/ActorFactories/ActorFactoryExponentialHeightFog.h b/Engine/Source/Editor/UnrealEd/Classes/ActorFactories/ActorFactoryExponentialHeightFog.h
--- a/Engine/Source/Editor/UnrealEd/Classes/ActorFactories/ActorFactoryExponentialHeightFog.h
+++ b/Engine/Source/Editor/UnrealEd/Classes/ActorFactories/ActorFactoryExponentialHeightFog.h
@@ -7,7 +7,9 @@
 UCLASS(MinimalAPI, config=Editor, collapsecategories, hidecategories=Object)
 class UActorFactoryExponentialHeightFog : public UActorFactory
 {
-	GENERATED_UCLASS_BODY()
+	GENERATED_BODY()
+public:
+	UNREALED_API UActorFactoryExponentialHeightFog(const FObjectInitializer& ObjectInitializer = FObjectInitializer::Get());
 };
 
 
diff --git a/Engine/Source/Editor/UnrealEd/Classes/ActorFactories/ActorFactoryInteractiveFoliage.h b/Engine/Source/Editor/UnrealEd/Classes/ActorFactories/ActorFactoryInteractiveFoliage.h
--- a/Engine/Source/Editor/UnrealEd/Classes/ActorFactories/ActorFactoryInteractiveFoliage.h
+++ b/Engine/Source/Editor/UnrealEd/Classes/ActorFactories/ActorFactoryInteractiveFoliage.h
@@ -7,7 +7,9 @@
 UCLASS(MinimalAPI, config=Editor)
 class UActorFactoryInteractiveFoliage : public UActorFactoryStaticMesh
 {
-	GENERATED_UCLASS_BODY()
+	GENERATED_BODY()
+public:
+	UNREALED_API UActorFactoryInteractiveFoliage(const FObjectInitializer& ObjectInitializer = FObjectInitializer::Get());
 
 
 	// Begin UActorFactory Interface
diff --git a/Engine/Source/Editor/UnrealEd/Classes/ActorFactories/ActorFactoryMatineeActor.h b/Engine/Source/Editor/UnrealEd/Classes/ActorFactories/ActorFactoryMatineeActor.h
--- a/Engine/Source/Editor/UnrealEd/Classes/ActorFactories/ActorFactoryMatineeActor.h
+++ b/Engine/Source/Editor/UnrealEd/Classes/ActorFactories/ActorFactoryMatineeActor.h
@@ -8,7 +8,9 @@
 UCLASS(MinimalAPI, config=Editor)
 class UActorFactoryMatineeActor : public UActorFactory
 {
-	GENERATED_UCLASS_BODY()
+	GENERATED_BODY()
+public:
+	UNREALED_API UActorFactoryMatineeActor(const FObjectInitializer& ObjectInitializer = FObjectInitializer::Get());
 
 	// Begin UActorFactory Interface
 	virtual void PostSpawnActor( UObject* Asset, AActor* NewActor ) override;
diff --git a/Engine/Source/Editor/UnrealEd/Classes/ActorFactories/ActorFactoryNiagara.h b/Engine/Source/Editor/UnrealEd/Classes/ActorFactories/ActorFactoryNiagara.h
--- a/Engine/Source/Editor/UnrealEd/Classes/ActorFactories/ActorFactoryNiagara.h
+++ b/Engine/Source/Editor/UnrealEd/Classes/ActorFactories/ActorFactoryNiagara.h
@@ -8,7 +8,9 @@
 UCLASS(MinimalAPI, config = Editor, collapsecategories, hidecategories = Object)
 class UActorFactoryNiagara : public UActorFactory
 {
-	GENERATED_UCLASS_BODY()
+	GENERATED_BODY()
+public:
+	UNREALED_API UActorFactoryNiagara(const FObjectInitializer& ObjectInitializer = FObjectInitializer::Get());
 
 	// Begin UActorFactory Interface
 	virtual void PostSpawnActor(UObject* Asset, AActor* NewActor) override;
diff --git a/Engine/Source/Editor/UnrealEd/Classes/ActorFactories/ActorFactoryNote.h b/Engine/Source/Editor/UnrealEd/Classes/ActorFactories/ActorFactoryNote.h
--- a/Engine/Source/Editor/UnrealEd/Classes/ActorFactories/ActorFactoryNote.h
+++ b/Engine/Source/Editor/UnrealEd/Classes/ActorFactories/ActorFactoryNote.h
@@ -6,6 +6,8 @@
 UCLASS( MinimalAPI, config = Editor, collapsecategories, hidecategories = Object )
 class UActorFactoryNote : public UActorFactory
 {
-	GENERATED_UCLASS_BODY()
+	GENERATED_BODY()
+public:
+	UNREALED_API UActorFactoryNote(const FObjectInitializer& ObjectInitializer = FObjectInitializer::Get());
 
 };
diff --git a/Engine/Source/Editor/UnrealEd/Classes/ActorFactories/ActorFactoryPawn.h b/Engine/Source/Editor/UnrealEd/Classes/ActorFactories/ActorFactoryPawn.h
--- a/Engine/Source/Editor/UnrealEd/Classes/ActorFactories/ActorFactoryPawn.h
+++ b/Engine/Source/Editor/UnrealEd/Classes/ActorFactories/ActorFactoryPawn.h
@@ -7,7 +7,9 @@
 UCLASS(MinimalAPI,config=Editor)
 class UActorFactoryPawn : public UActorFactoryEmptyActor
 {
-	GENERATED_UCLASS_BODY()
+	GENERATED_BODY()
+public:
+	UNREALED_API UActorFactoryPawn(const FObjectInitializer& ObjectInitializer = FObjectInitializer::Get());
 
 public:
 	virtual bool CanCreateActorFrom( const FAssetData& AssetData, FText& OutErrorMsg ) override;
diff --git a/Engine/Source/Editor/UnrealEd/Classes/ActorFactories/ActorFactoryPhysicsAsset.h b/Engine/Source/Editor/UnrealEd/Classes/ActorFactories/ActorFactoryPhysicsAsset.h
--- a/Engine/Source/Editor/UnrealEd/Classes/ActorFactories/ActorFactoryPhysicsAsset.h
+++ b/Engine/Source/Editor/UnrealEd/Classes/ActorFactories/ActorFactoryPhysicsAsset.h
@@ -8,7 +8,9 @@
 UCLASS(MinimalAPI, config=Editor, collapsecategories, hidecategories=Object)
 class UActorFactoryPhysicsAsset : public UActorFactory
 {
-	GENERATED_UCLASS_BODY()
+	GENERATED_BODY()
+public:
+	UNREALED_API UActorFactoryPhysicsAsset(const FObjectInitializer& ObjectInitializer = FObjectInitializer::Get());
 
 	// Begin UActorFactory Interface
 	virtual bool PreSpawnActor( UObject* Asset, FTransform& InOutLocation ) override;
diff --git a/Engine/Source/Editor/UnrealEd/Classes/ActorFactories/ActorFactoryPlaneReflectionCapture.h b/Engine/Source/Editor/UnrealEd/Classes/ActorFactories/ActorFactoryPlaneReflectionCapture.h
--- a/Engine/Source/Editor/UnrealEd/Classes/ActorFactories/ActorFactoryPlaneReflectionCapture.h
+++ b/Engine/Source/Editor/UnrealEd/Classes/ActorFactories/ActorFactoryPlaneReflectionCapture.h
@@ -7,7 +7,9 @@
 UCLASS(MinimalAPI, config=Editor, collapsecategories, hidecategories=Object)
 class UActorFactoryPlaneReflectionCapture : public UActorFactory
 {
-	GENERATED_UCLASS_BODY()
+	GENERATED_BODY()
+public:
+	UNREALED_API UActorFactoryPlaneReflectionCapture(const FObjectInitializer& ObjectInitializer = FObjectInitializer::Get());
 };
 
 
diff --git a/Engine/Source/Editor/UnrealEd/Classes/ActorFactories/ActorFactoryPlayerStart.h b/Engine/Source/Editor/UnrealEd/Classes/ActorFactories/ActorFactoryPlayerStart.h
--- a/Engine/Source/Editor/UnrealEd/Classes/ActorFactories/ActorFactoryPlayerStart.h
+++ b/Engine/Source/Editor/UnrealEd/Classes/ActorFactories/ActorFactoryPlayerStart.h
@@ -7,7 +7,9 @@
 UCLASS(MinimalAPI, config=Editor, collapsecategories, hidecategories=Object)
 class UActorFactoryPlayerStart : public UActorFactory
 {
-	GENERATED_UCLASS_BODY()
+	GENERATED_BODY()
+public:
+	UNREALED_API UActorFactoryPlayerStart(const FObjectInitializer& ObjectInitializer = FObjectInitializer::Get());
 
 };
 
diff --git a/Engine/Source/Editor/UnrealEd/Classes/ActorFactories/ActorFactoryPointLight.h b/Engine/Source/Editor/UnrealEd/Classes/ActorFactories/ActorFactoryPointLight.h
--- a/Engine/Source/Editor/UnrealEd/Classes/ActorFactories/ActorFactoryPointLight.h
+++ b/Engine/Source/Editor/UnrealEd/Classes/ActorFactories/ActorFactoryPointLight.h
@@ -7,7 +7,9 @@
 UCLASS(MinimalAPI, config=Editor, collapsecategories, hidecategories=Object)
 class UActorFactoryPointLight : public UActorFactory
 {
-	GENERATED_UCLASS_BODY()
+	GENERATED_BODY()
+public:
+	UNREALED_API UActorFactoryPointLight(const FObjectInitializer& ObjectInitializer = FObjectInitializer::Get());
 
 };
 
diff --git a/Engine/Source/Editor/UnrealEd/Classes/ActorFactories/ActorFactorySkeletalMesh.h b/Engine/Source/Editor/UnrealEd/Classes/ActorFactories/ActorFactorySkeletalMesh.h
--- a/Engine/Source/Editor/UnrealEd/Classes/ActorFactories/ActorFactorySkeletalMesh.h
+++ b/Engine/Source/Editor/UnrealEd/Classes/ActorFactories/ActorFactorySkeletalMesh.h
@@ -10,7 +10,9 @@ class USkeletalMesh;
 UCLASS(MinimalAPI, config=Editor)
 class UActorFactorySkeletalMesh : public UActorFactory
 {
-	GENERATED_UCLASS_BODY()
+	GENERATED_BODY()
+public:
+	UNREALED_API UActorFactorySkeletalMesh(const FObjectInitializer& ObjectInitializer = FObjectInitializer::Get());
 
 protected:
 	// Begin UActorFactory Interface
diff --git a/Engine/Source/Editor/UnrealEd/Classes/ActorFactories/ActorFactorySkyLight.h b/Engine/Source/Editor/UnrealEd/Classes/ActorFactories/ActorFactorySkyLight.h
--- a/Engine/Source/Editor/UnrealEd/Classes/ActorFactories/ActorFactorySkyLight.h
+++ b/Engine/Source/Editor/UnrealEd/Classes/ActorFactories/ActorFactorySkyLight.h
@@ -7,7 +7,9 @@
 UCLASS(MinimalAPI, config=Editor, collapsecategories, hidecategories=Object)
 class UActorFactorySkyLight : public UActorFactory
 {
-	GENERATED_UCLASS_BODY()
+	GENERATED_BODY()
+public:
+	UNREALED_API UActorFactorySkyLight(const FObjectInitializer& ObjectInitializer = FObjectInitializer::Get());
 
 };
 
diff --git a/Engine/Source/Editor/UnrealEd/Classes/ActorFactories/ActorFactorySphereReflectionCapture.h b/Engine/Source/Editor/UnrealEd/Classes/ActorFactories/ActorFactorySphereReflectionCapture.h
--- a/Engine/Source/Editor/UnrealEd/Classes/ActorFactories/ActorFactorySphereReflectionCapture.h
+++ b/Engine/Source/Editor/UnrealEd/Classes/ActorFactories/ActorFactorySphereReflectionCapture.h
@@ -7,7 +7,9 @@
 UCLASS(MinimalAPI, config=Editor, collapsecategories, hidecategories=Object)
 class UActorFactorySphereReflectionCapture : public UActorFactory
 {
-	GENERATED_UCLASS_BODY()
+	GENERATED_BODY()
+public:
+	UNREALED_API UActorFactorySphereReflectionCapture(const FObjectInitializer& ObjectInitializer = FObjectInitializer::Get());
 };
 
 
diff --git a/Engine/Source/Editor/UnrealEd/Classes/ActorFactories/ActorFactorySphereVolume.h b/Engine/Source/Editor/UnrealEd/Classes/ActorFactories/ActorFactorySphereVolume.h
--- a/Engine/Source/Editor/UnrealEd/Classes/ActorFactories/ActorFactorySphereVolume.h
+++ b/Engine/Source/Editor/UnrealEd/Classes/ActorFactories/ActorFactorySphereVolume.h
@@ -7,7 +7,9 @@
 UCLASS(MinimalAPI, config=Editor, collapsecategories, hidecategories=Object)
 class UActorFactorySphereVolume : public UActorFactory
 {
-	GENERATED_UCLASS_BODY()
+	GENERATED_BODY()
+public:
+	UNREALED_API UActorFactorySphereVolume(const FObjectInitializer& ObjectInitializer = FObjectInitializer::Get());
 
 	// Begin UActorFactory Interface
 	virtual bool CanCreateActorFrom( const FAssetData& AssetData, FText& OutErrorMsg ) override;
diff --git a/Engine/Source/Editor/UnrealEd/Classes/ActorFactories/ActorFactorySpotLight.h b/Engine/Source/Editor/UnrealEd/Classes/ActorFactories/ActorFactorySpotLight.h
--- a/Engine/Source/Editor/UnrealEd/Classes/ActorFactories/ActorFactorySpotLight.h
+++ b/Engine/Source/Editor/UnrealEd/Classes/ActorFactories/ActorFactorySpotLight.h
@@ -7,7 +7,9 @@
 UCLASS(MinimalAPI, config=Editor, collapsecategories, hidecategories=Object)
 class UActorFactorySpotLight : public UActorFactory
 {
-	GENERATED_UCLASS_BODY()
+	GENERATED_BODY()
+public:
+	UNREALED_API UActorFactorySpotLight(const FObjectInitializer& ObjectInitializer = FObjectInitializer::Get());
 
 };
 
diff --git a/Engine/Source/Editor/UnrealEd/Classes/ActorFactories/ActorFactoryStaticMesh.h b/Engine/Source/Editor/UnrealEd/Classes/ActorFactories/ActorFactoryStaticMesh.h
--- a/Engine/Source/Editor/UnrealEd/Classes/ActorFactories/ActorFactoryStaticMesh.h
+++ b/Engine/Source/Editor/UnrealEd/Classes/ActorFactories/ActorFactoryStaticMesh.h
@@ -7,7 +7,9 @@
 UCLASS(MinimalAPI, config=Editor)
 class UActorFactoryStaticMesh : public UActorFactory
 {
-	GENERATED_UCLASS_BODY()
+	GENERATED_BODY()
+public:
+	UNREALED_API UActorFactoryStaticMesh(const FObjectInitializer& ObjectInitializer = FObjectInitializer::Get());
 
 	// Begin UActorFactory Interface
 	virtual bool CanCreateActorFrom( const FAssetData& AssetData, FText& OutErrorMsg ) override;
diff --git a/Engine/Source/Editor/UnrealEd/Classes/ActorFactories/ActorFactoryTargetPoint.h b/Engine/Source/Editor/UnrealEd/Classes/ActorFactories/ActorFactoryTargetPoint.h
--- a/Engine/Source/Editor/UnrealEd/Classes/ActorFactories/ActorFactoryTargetPoint.h
+++ b/Engine/Source/Editor/UnrealEd/Classes/ActorFactories/ActorFactoryTargetPoint.h
@@ -6,6 +6,8 @@
 UCLASS( MinimalAPI, config = Editor, collapsecategories, hidecategories = Object )
 class UActorFactoryTargetPoint : public UActorFactory
 {
-	GENERATED_UCLASS_BODY()
+	GENERATED_BODY()
+public:
+	UNREALED_API UActorFactoryTargetPoint(const FObjectInitializer& ObjectInitializer = FObjectInitializer::Get());
 
 };
diff --git a/Engine/Source/Editor/UnrealEd/Classes/ActorFactories/ActorFactoryTextRender.h b/Engine/Source/Editor/UnrealEd/Classes/ActorFactories/ActorFactoryTextRender.h
--- a/Engine/Source/Editor/UnrealEd/Classes/ActorFactories/ActorFactoryTextRender.h
+++ b/Engine/Source/Editor/UnrealEd/Classes/ActorFactories/ActorFactoryTextRender.h
@@ -7,7 +7,9 @@
 UCLASS(MinimalAPI, config=Editor)
 class UActorFactoryTextRender : public UActorFactory
 {
-	GENERATED_UCLASS_BODY()
+	GENERATED_BODY()
+public:
+	UNREALED_API UActorFactoryTextRender(const FObjectInitializer& ObjectInitializer = FObjectInitializer::Get());
 };
 
 
diff --git a/Engine/Source/Editor/UnrealEd/Classes/ActorFactories/ActorFactoryTriggerBox.h b/Engine/Source/Editor/UnrealEd/Classes/ActorFactories/ActorFactoryTriggerBox.h
--- a/Engine/Source/Editor/UnrealEd/Classes/ActorFactories/ActorFactoryTriggerBox.h
+++ b/Engine/Source/Editor/UnrealEd/Classes/ActorFactories/ActorFactoryTriggerBox.h
@@ -7,7 +7,9 @@
 UCLASS(MinimalAPI, config=Editor, collapsecategories, hidecategories=Object)
 class UActorFactoryTriggerBox : public UActorFactory
 {
-	GENERATED_UCLASS_BODY()
+	GENERATED_BODY()
+public:
+	UNREALED_API UActorFactoryTriggerBox(const FObjectInitializer& ObjectInitializer = FObjectInitializer::Get());
 
 };
 
diff --git a/Engine/Source/Editor/UnrealEd/Classes/ActorFactories/ActorFactoryTriggerCapsule.h b/Engine/Source/Editor/UnrealEd/Classes/ActorFactories/ActorFactoryTriggerCapsule.h
--- a/Engine/Source/Editor/UnrealEd/Classes/ActorFactories/ActorFactoryTriggerCapsule.h
+++ b/Engine/Source/Editor/UnrealEd/Classes/ActorFactories/ActorFactoryTriggerCapsule.h
@@ -7,7 +7,9 @@
 UCLASS(MinimalAPI, config=Editor, collapsecategories, hidecategories=Object)
 class UActorFactoryTriggerCapsule : public UActorFactory
 {
-	GENERATED_UCLASS_BODY()
+	GENERATED_BODY()
+public:
+	UNREALED_API UActorFactoryTriggerCapsule(const FObjectInitializer& ObjectInitializer = FObjectInitializer::Get());
 
 };
 
diff --git a/Engine/Source/Editor/UnrealEd/Classes/ActorFactories/ActorFactoryTriggerSphere.h b/Engine/Source/Editor/UnrealEd/Classes/ActorFactories/ActorFactoryTriggerSphere.h
--- a/Engine/Source/Editor/UnrealEd/Classes/ActorFactories/ActorFactoryTriggerSphere.h
+++ b/Engine/Source/Editor/UnrealEd/Classes/ActorFactories/ActorFactoryTriggerSphere.h
@@ -7,7 +7,9 @@
 UCLASS(MinimalAPI, config=Editor, collapsecategories, hidecategories=Object)
 class UActorFactoryTriggerSphere : public UActorFactory
 {
-	GENERATED_UCLASS_BODY()
+	GENERATED_BODY()
+public:
+	UNREALED_API UActorFactoryTriggerSphere(const FObjectInitializer& ObjectInitializer = FObjectInitializer::Get());
 
 };
 
diff --git a/Engine/Source/Editor/UnrealEd/Classes/ActorFactories/ActorFactoryVectorFieldVolume.h b/Engine/Source/Editor/UnrealEd/Classes/ActorFactories/ActorFactoryVectorFieldVolume.h
--- a/Engine/Source/Editor/UnrealEd/Classes/ActorFactories/ActorFactoryVectorFieldVolume.h
+++ b/Engine/Source/Editor/UnrealEd/Classes/ActorFactories/ActorFactoryVectorFieldVolume.h
@@ -7,7 +7,9 @@
 UCLASS(MinimalAPI, config=Editor, collapsecategories, hidecategories=Object)
 class UActorFactoryVectorFieldVolume : public UActorFactory
 {
-	GENERATED_UCLASS_BODY()
+	GENERATED_BODY()
+public:
+	UNREALED_API UActorFactoryVectorFieldVolume(const FObjectInitializer& ObjectInitializer = FObjectInitializer::Get());
 
 	// Begin UActorFactory Interface
 	virtual void PostSpawnActor( UObject* Asset, AActor* NewActor ) override;
diff --git a/Engine/Source/Editor/UnrealEd/Classes/Animation/DebugSkelMeshComponent.h b/Engine/Source/Editor/UnrealEd/Classes/Animation/DebugSkelMeshComponent.h
--- a/Engine/Source/Editor/UnrealEd/Classes/Animation/DebugSkelMeshComponent.h
+++ b/Engine/Source/Editor/UnrealEd/Classes/Animation/DebugSkelMeshComponent.h
@@ -8,7 +8,7 @@
 USTRUCT()
 struct FSelectedSocketInfo
 {
-	GENERATED_USTRUCT_BODY()
+	GENERATED_BODY()
 
 	/** Default constructor */
 	FSelectedSocketInfo()
@@ -38,7 +38,9 @@ struct FSelectedSocketInfo
 UCLASS(transient, MinimalAPI)
 class UDebugSkelMeshComponent : public USkeletalMeshComponent
 {
-	GENERATED_UCLASS_BODY()
+	GENERATED_BODY()
+public:
+	UNREALED_API UDebugSkelMeshComponent(const FObjectInitializer& ObjectInitializer = FObjectInitializer::Get());
 
 	/** If true, render a wireframe skeleton of the mesh animated with the raw (uncompressed) animation data. */
 	UPROPERTY()
diff --git a/Engine/Source/Editor/UnrealEd/Classes/Animation/EditorAnimBaseObj.h b/Engine/Source/Editor/UnrealEd/Classes/Animation/EditorAnimBaseObj.h
--- a/Engine/Source/Editor/UnrealEd/Classes/Animation/EditorAnimBaseObj.h
+++ b/Engine/Source/Editor/UnrealEd/Classes/Animation/EditorAnimBaseObj.h
@@ -16,7 +16,9 @@ DECLARE_DELEGATE_TwoParams( FOnAnimObjectChange, class UObject*, bool)
 UCLASS(abstract, MinimalAPI)
 class UEditorAnimBaseObj: public UObject
 {
-	GENERATED_UCLASS_BODY()
+	GENERATED_BODY()
+public:
+	UNREALED_API UEditorAnimBaseObj(const FObjectInitializer& ObjectInitializer = FObjectInitializer::Get());
 public:
 	
 
diff --git a/Engine/Source/Editor/UnrealEd/Classes/Animation/EditorAnimCompositeSegment.h b/Engine/Source/Editor/UnrealEd/Classes/Animation/EditorAnimCompositeSegment.h
--- a/Engine/Source/Editor/UnrealEd/Classes/Animation/EditorAnimCompositeSegment.h
+++ b/Engine/Source/Editor/UnrealEd/Classes/Animation/EditorAnimCompositeSegment.h
@@ -16,7 +16,9 @@ DECLARE_DELEGATE_OneParam( FOnAnimSegmentChanged, class UEditorAnimSegment*)
 UCLASS(hidecategories=UObject, MinimalAPI, BlueprintType)
 class UEditorAnimCompositeSegment: public UEditorAnimBaseObj
 {
-	GENERATED_UCLASS_BODY()
+	GENERATED_BODY()
+public:
+	UNREALED_API UEditorAnimCompositeSegment(const FObjectInitializer& ObjectInitializer = FObjectInitializer::Get());
 public:
 
 	/** Default blend in time. */
diff --git a/Engine/Source/Editor/UnrealEd/Classes/Animation/EditorAnimSegment.h b/Engine/Source/Editor/UnrealEd/Classes/Animation/EditorAnimSegment.h
--- a/Engine/Source/Editor/UnrealEd/Classes/Animation/EditorAnimSegment.h
+++ b/Engine/Source/Editor/UnrealEd/Classes/Animation/EditorAnimSegment.h
@@ -16,7 +16,9 @@ DECLARE_DELEGATE_OneParam( FOnAnimSegmentChanged, class UEditorAnimSegment*)
 UCLASS(hidecategories=UObject, MinimalAPI, BlueprintType)
 class UEditorAnimSegment: public UEditorAnimBaseObj
 {
-	GENERATED_UCLASS_BODY()
+	GENERATED_BODY()
+public:
+	UNREALED_API UEditorAnimSegment(const FObjectInitializer& ObjectInitializer = FObjectInitializer::Get());
 public:
 
 	/** Default blend in time. */
diff --git a/Engine/Source/Editor/UnrealEd/Classes/Animation/EditorCompositeSection.h b/Engine/Source/Editor/UnrealEd/Classes/Animation/EditorCompositeSection.h
--- a/Engine/Source/Editor/UnrealEd/Classes/Animation/EditorCompositeSection.h
+++ b/Engine/Source/Editor/UnrealEd/Classes/Animation/EditorCompositeSection.h
@@ -13,7 +13,9 @@
 UCLASS(hidecategories=UObject, MinimalAPI, BlueprintType)
 class UEditorCompositeSection: public UEditorAnimBaseObj
 {
-	GENERATED_UCLASS_BODY()
+	GENERATED_BODY()
+public:
+	UNREALED_API UEditorCompositeSection(const FObjectInitializer& ObjectInitializer = FObjectInitializer::Get());
 public:
 
 	/** Default blend in time. */
diff --git a/Engine/Source/Editor/UnrealEd/Classes/Animation/EditorNotifyObject.h b/Engine/Source/Editor/UnrealEd/Classes/Animation/EditorNotifyObject.h
--- a/Engine/Source/Editor/UnrealEd/Classes/Animation/EditorNotifyObject.h
+++ b/Engine/Source/Editor/UnrealEd/Classes/Animation/EditorNotifyObject.h
@@ -14,7 +14,9 @@ class UEditorNotifyEntry;
 UCLASS(MinimalAPI)
 class UEditorNotifyObject : public UEditorAnimBaseObj
 {
-	GENERATED_UCLASS_BODY()
+	GENERATED_BODY()
+public:
+	UNREALED_API UEditorNotifyObject(const FObjectInitializer& ObjectInitializer = FObjectInitializer::Get());
 
 	/** Set up the editor object
 	 *	@param TrackIdx Index of the track the notify belongs to
@@ -34,4 +36,4 @@ class UEditorNotifyObject : public UEditorAnimBaseObj
 	/** The notify event to modify */
 	UPROPERTY(EditAnywhere, Category=Event)
 	FAnimNotifyEvent Event;
-};
\ No newline at end of file
+};
diff --git a/Engine/Source/Editor/UnrealEd/Classes/Animation/EditorParentPlayerListObj.h b/Engine/Source/Editor/UnrealEd/Classes/Animation/EditorParentPlayerListObj.h
--- a/Engine/Source/Editor/UnrealEd/Classes/Animation/EditorParentPlayerListObj.h
+++ b/Engine/Source/Editor/UnrealEd/Classes/Animation/EditorParentPlayerListObj.h
@@ -11,7 +11,9 @@ class UAnimBlueprint;
 UCLASS()
 class UNREALED_API UEditorParentPlayerListObj : public UObject
 {
-	GENERATED_UCLASS_BODY()
+	GENERATED_BODY()
+public:
+	UEditorParentPlayerListObj(const FObjectInitializer& ObjectInitializer = FObjectInitializer::Get());
 public:
 	// Build a possible list of overrides from a blueprint and pull in any existing overrides.
 	void InitialiseFromBlueprint(UAnimBlueprint* Blueprint);
diff --git a/Engine/Source/Editor/UnrealEd/Classes/Animation/EditorSkeletonNotifyObj.h b/Engine/Source/Editor/UnrealEd/Classes/Animation/EditorSkeletonNotifyObj.h
--- a/Engine/Source/Editor/UnrealEd/Classes/Animation/EditorSkeletonNotifyObj.h
+++ b/Engine/Source/Editor/UnrealEd/Classes/Animation/EditorSkeletonNotifyObj.h
@@ -13,7 +13,7 @@ DECLARE_DELEGATE_TwoParams( FOnAnimObjectChange, class UObject*, bool)
 USTRUCT()
 struct FSkeletonNotifyDependentAnimations
 {
-	GENERATED_USTRUCT_BODY()
+	GENERATED_BODY()
 
 	UPROPERTY(EditAnywhere, BlueprintReadWrite, Category=SkeletonNotifies)
 	TArray<UAnimSequenceBase*> Animations;
@@ -26,7 +26,9 @@ struct FSkeletonNotifyDependentAnimations
 UCLASS(MinimalAPI)
 class UEditorSkeletonNotifyObj : public UObject
 {
-	GENERATED_UCLASS_BODY()
+	GENERATED_BODY()
+public:
+	UNREALED_API UEditorSkeletonNotifyObj(const FObjectInitializer& ObjectInitializer = FObjectInitializer::Get());
 public:
 	
 	UObject* EditedObject;
diff --git a/Engine/Source/Editor/UnrealEd/Classes/Builders/ConeBuilder.h b/Engine/Source/Editor/UnrealEd/Classes/Builders/ConeBuilder.h
--- a/Engine/Source/Editor/UnrealEd/Classes/Builders/ConeBuilder.h
+++ b/Engine/Source/Editor/UnrealEd/Classes/Builders/ConeBuilder.h
@@ -10,11 +10,9 @@
 UCLASS(MinimalAPI, autoexpandcategories=BrushSettings, EditInlineNew, meta=(DisplayName="Cone"))
 class UConeBuilder : public UEditorBrushBuilder
 {
-public:
 	GENERATED_BODY()
-
 public:
-	UConeBuilder(const FObjectInitializer& ObjectInitializer = FObjectInitializer::Get());
+	UNREALED_API UConeBuilder(const FObjectInitializer& ObjectInitializer = FObjectInitializer::Get());
 
 	/** Distance from base to tip of cone */
 	UPROPERTY(EditAnywhere, Category=BrushSettings, meta=(ClampMin = "0.000001"))
diff --git a/Engine/Source/Editor/UnrealEd/Classes/Builders/CubeBuilder.h b/Engine/Source/Editor/UnrealEd/Classes/Builders/CubeBuilder.h
--- a/Engine/Source/Editor/UnrealEd/Classes/Builders/CubeBuilder.h
+++ b/Engine/Source/Editor/UnrealEd/Classes/Builders/CubeBuilder.h
@@ -10,11 +10,9 @@
 UCLASS(MinimalAPI, autoexpandcategories=BrushSettings, EditInlineNew, meta=(DisplayName="Box"))
 class UCubeBuilder : public UEditorBrushBuilder
 {
-public:
 	GENERATED_BODY()
-
 public:
-	UCubeBuilder(const FObjectInitializer& ObjectInitializer = FObjectInitializer::Get());
+	UNREALED_API UCubeBuilder(const FObjectInitializer& ObjectInitializer = FObjectInitializer::Get());
 
 	/** The size of the cube in the X dimension */
 	UPROPERTY(EditAnywhere, Category=BrushSettings, meta=(ClampMin = "0.000001"))
diff --git a/Engine/Source/Editor/UnrealEd/Classes/Builders/CurvedStairBuilder.h b/Engine/Source/Editor/UnrealEd/Classes/Builders/CurvedStairBuilder.h
--- a/Engine/Source/Editor/UnrealEd/Classes/Builders/CurvedStairBuilder.h
+++ b/Engine/Source/Editor/UnrealEd/Classes/Builders/CurvedStairBuilder.h
@@ -10,11 +10,9 @@
 UCLASS(MinimalAPI, autoexpandcategories=BrushSettings, EditInlineNew,meta=(DisplayName="Curved Stair"))
 class UCurvedStairBuilder : public UEditorBrushBuilder
 {
-public:
 	GENERATED_BODY()
-
 public:
-	UCurvedStairBuilder(const FObjectInitializer& ObjectInitializer = FObjectInitializer::Get());
+	UNREALED_API UCurvedStairBuilder(const FObjectInitializer& ObjectInitializer = FObjectInitializer::Get());
 
 	/** The radius of the inner curve of the stair */
 	UPROPERTY(EditAnywhere, Category=BrushSettings, meta=(ClampMin = "1"))
diff --git a/Engine/Source/Editor/UnrealEd/Classes/Builders/CylinderBuilder.h b/Engine/Source/Editor/UnrealEd/Classes/Builders/CylinderBuilder.h
--- a/Engine/Source/Editor/UnrealEd/Classes/Builders/CylinderBuilder.h
+++ b/Engine/Source/Editor/UnrealEd/Classes/Builders/CylinderBuilder.h
@@ -15,11 +15,9 @@
 UCLASS(MinimalAPI, autoexpandcategories=BrushSettings, EditInlineNew, meta=(DisplayName="Cylinder"))
 class UCylinderBuilder : public UEditorBrushBuilder
 {
-public:
 	GENERATED_BODY()
-
 public:
-	UCylinderBuilder(const FObjectInitializer& ObjectInitializer = FObjectInitializer::Get());
+	UNREALED_API UCylinderBuilder(const FObjectInitializer& ObjectInitializer = FObjectInitializer::Get());
 
 	/** Distance from base to tip of cylinder */
 	UPROPERTY(EditAnywhere, Category=BrushSettings, meta=(ClampMin = "0.000001"))
diff --git a/Engine/Source/Editor/UnrealEd/Classes/Builders/EditorBrushBuilder.h b/Engine/Source/Editor/UnrealEd/Classes/Builders/EditorBrushBuilder.h
--- a/Engine/Source/Editor/UnrealEd/Classes/Builders/EditorBrushBuilder.h
+++ b/Engine/Source/Editor/UnrealEd/Classes/Builders/EditorBrushBuilder.h
@@ -10,8 +10,7 @@ class UEditorBrushBuilder : public UBrushBuilder
 public:
 	GENERATED_BODY()
 public:
-
-	UEditorBrushBuilder(const FObjectInitializer& ObjectInitializer = FObjectInitializer::Get());
+	UNREALED_API UEditorBrushBuilder(const FObjectInitializer& ObjectInitializer = FObjectInitializer::Get());
 
 	/** UBrushBuilder interface */
 	virtual void BeginBrush( bool InMergeCoplanars, FName InLayer ) override;
diff --git a/Engine/Source/Editor/UnrealEd/Classes/Builders/LinearStairBuilder.h b/Engine/Source/Editor/UnrealEd/Classes/Builders/LinearStairBuilder.h
--- a/Engine/Source/Editor/UnrealEd/Classes/Builders/LinearStairBuilder.h
+++ b/Engine/Source/Editor/UnrealEd/Classes/Builders/LinearStairBuilder.h
@@ -10,11 +10,9 @@
 UCLASS(MinimalAPI, autoexpandcategories=BrushSettings, EditInlineNew, meta=(DisplayName="Linear Stair"))
 class ULinearStairBuilder : public UEditorBrushBuilder
 {
-public:
 	GENERATED_BODY()
-
 public:
-	ULinearStairBuilder(const FObjectInitializer& ObjectInitializer = FObjectInitializer::Get());
+	UNREALED_API ULinearStairBuilder(const FObjectInitializer& ObjectInitializer = FObjectInitializer::Get());
 
 	/** The length of each step */
 	UPROPERTY(EditAnywhere, Category=BrushSettings, meta=(ClampMin = "1"))
diff --git a/Engine/Source/Editor/UnrealEd/Classes/Builders/SheetBuilder.h b/Engine/Source/Editor/UnrealEd/Classes/Builders/SheetBuilder.h
--- a/Engine/Source/Editor/UnrealEd/Classes/Builders/SheetBuilder.h
+++ b/Engine/Source/Editor/UnrealEd/Classes/Builders/SheetBuilder.h
@@ -19,11 +19,9 @@ enum ESheetAxis
 UCLASS(MinimalAPI, autoexpandcategories=BrushSettings, EditInlineNew, NotPlaceable, meta=(DisplayName="Plane"))
 class USheetBuilder : public UEditorBrushBuilder
 {
-public:
 	GENERATED_BODY()
-
 public:
-	USheetBuilder(const FObjectInitializer& ObjectInitializer = FObjectInitializer::Get());
+	UNREALED_API USheetBuilder(const FObjectInitializer& ObjectInitializer = FObjectInitializer::Get());
 
 	UPROPERTY(EditAnywhere, Category=BrushSettings, meta=(ClampMin = "1"))
 	int32 X;
diff --git a/Engine/Source/Editor/UnrealEd/Classes/Builders/SpiralStairBuilder.h b/Engine/Source/Editor/UnrealEd/Classes/Builders/SpiralStairBuilder.h
--- a/Engine/Source/Editor/UnrealEd/Classes/Builders/SpiralStairBuilder.h
+++ b/Engine/Source/Editor/UnrealEd/Classes/Builders/SpiralStairBuilder.h
@@ -10,11 +10,9 @@
 UCLASS(MinimalAPI, autoexpandcategories=BrushSettings, EditInlineNew, meta=(DisplayName="Spiral Stair"))
 class USpiralStairBuilder : public UEditorBrushBuilder
 {
-public:
 	GENERATED_BODY()
-
 public:
-	USpiralStairBuilder(const FObjectInitializer& ObjectInitializer = FObjectInitializer::Get());
+	UNREALED_API USpiralStairBuilder(const FObjectInitializer& ObjectInitializer = FObjectInitializer::Get());
 
 	/** The radius of the inner curve of the stair */
 	UPROPERTY(EditAnywhere, Category=BrushSettings, meta=(ClampMin = "1"))
diff --git a/Engine/Source/Editor/UnrealEd/Classes/Builders/TetrahedronBuilder.h b/Engine/Source/Editor/UnrealEd/Classes/Builders/TetrahedronBuilder.h
--- a/Engine/Source/Editor/UnrealEd/Classes/Builders/TetrahedronBuilder.h
+++ b/Engine/Source/Editor/UnrealEd/Classes/Builders/TetrahedronBuilder.h
@@ -10,11 +10,9 @@
 UCLASS(MinimalAPI, autoexpandcategories=BrushSettings, EditInlineNew, meta=(DisplayName="Sphere"))
 class UTetrahedronBuilder : public UEditorBrushBuilder
 {
-public:
 	GENERATED_BODY()
-
 public:
-	UTetrahedronBuilder(const FObjectInitializer& ObjectInitializer = FObjectInitializer::Get());
+	UNREALED_API UTetrahedronBuilder(const FObjectInitializer& ObjectInitializer = FObjectInitializer::Get());
 
 	/** The radius of this sphere */
 	UPROPERTY(EditAnywhere, Category=BrushSettings, meta=(ClampMin = "0.000001"))
diff --git a/Engine/Source/Editor/UnrealEd/Classes/Builders/VolumetricBuilder.h b/Engine/Source/Editor/UnrealEd/Classes/Builders/VolumetricBuilder.h
--- a/Engine/Source/Editor/UnrealEd/Classes/Builders/VolumetricBuilder.h
+++ b/Engine/Source/Editor/UnrealEd/Classes/Builders/VolumetricBuilder.h
@@ -10,12 +10,10 @@
 UCLASS(MinimalAPI, autoexpandcategories=BrushSettings, EditInlineNew, NotPlaceable, meta=(DisplayName="Volumetric"))
 class UVolumetricBuilder : public UEditorBrushBuilder
 {
-public:
 	GENERATED_BODY()
-
 public:
-	UVolumetricBuilder(const FObjectInitializer& ObjectInitializer = FObjectInitializer::Get());
-	
+	UNREALED_API UVolumetricBuilder(const FObjectInitializer& ObjectInitializer = FObjectInitializer::Get());
+
 	UPROPERTY(EditAnywhere, Category=BrushSettings, meta=(ClampMin = "0.000001"))
 	float Z;
 
diff --git a/Engine/Source/Editor/UnrealEd/Classes/Commandlets/AutomatedLauncherCommandlet.h b/Engine/Source/Editor/UnrealEd/Classes/Commandlets/AutomatedLauncherCommandlet.h
--- a/Engine/Source/Editor/UnrealEd/Classes/Commandlets/AutomatedLauncherCommandlet.h
+++ b/Engine/Source/Editor/UnrealEd/Classes/Commandlets/AutomatedLauncherCommandlet.h
@@ -9,9 +9,11 @@
 UCLASS()
 class UAutomatedLauncherCommandlet : public UCommandlet
 {
-	GENERATED_UCLASS_BODY()
+	GENERATED_BODY()
+public:
+	UAutomatedLauncherCommandlet(const FObjectInitializer& ObjectInitializer = FObjectInitializer::Get());
 
 	// Begin UCommandlet Interface
 	virtual int32 Main(const FString& Params) override;
 	// End UCommandlet Interface
-};
\ No newline at end of file
+};
diff --git a/Engine/Source/Editor/UnrealEd/Classes/Commandlets/CompressAnimationsCommandlet.h b/Engine/Source/Editor/UnrealEd/Classes/Commandlets/CompressAnimationsCommandlet.h
--- a/Engine/Source/Editor/UnrealEd/Classes/Commandlets/CompressAnimationsCommandlet.h
+++ b/Engine/Source/Editor/UnrealEd/Classes/Commandlets/CompressAnimationsCommandlet.h
@@ -7,7 +7,9 @@
 UCLASS()
 class UCompressAnimationsCommandlet : public UCommandlet
 {
-	GENERATED_UCLASS_BODY()
+	GENERATED_BODY()
+public:
+	UCompressAnimationsCommandlet(const FObjectInitializer& ObjectInitializer = FObjectInitializer::Get());
 	// Begin UCommandlet Interface
 	virtual int32 Main(const FString& Params) override;
 	// End UCommandlet Interface
diff --git a/Engine/Source/Editor/UnrealEd/Classes/Commandlets/CookCommandlet.h b/Engine/Source/Editor/UnrealEd/Classes/Commandlets/CookCommandlet.h
--- a/Engine/Source/Editor/UnrealEd/Classes/Commandlets/CookCommandlet.h
+++ b/Engine/Source/Editor/UnrealEd/Classes/Commandlets/CookCommandlet.h
@@ -12,7 +12,9 @@ UCLASS(config=Editor)
 class UCookCommandlet
 	: public UCommandlet
 {
-	GENERATED_UCLASS_BODY()
+	GENERATED_BODY()
+public:
+	UCookCommandlet(const FObjectInitializer& ObjectInitializer = FObjectInitializer::Get());
 
 	/** List of asset types that will force GC after loading them during cook */
 	UPROPERTY(config)
@@ -157,4 +159,4 @@ private:
 			InOutFilesToCook.AddUnique(InFilename);
 		}
 	}
-};
\ No newline at end of file
+};
diff --git a/Engine/Source/Editor/UnrealEd/Classes/Commandlets/DerivedDataCacheCommandlet.h b/Engine/Source/Editor/UnrealEd/Classes/Commandlets/DerivedDataCacheCommandlet.h
--- a/Engine/Source/Editor/UnrealEd/Classes/Commandlets/DerivedDataCacheCommandlet.h
+++ b/Engine/Source/Editor/UnrealEd/Classes/Commandlets/DerivedDataCacheCommandlet.h
@@ -12,7 +12,9 @@
 UCLASS()
 class UDerivedDataCacheCommandlet : public UCommandlet
 {
-	GENERATED_UCLASS_BODY()
+	GENERATED_BODY()
+public:
+	UDerivedDataCacheCommandlet(const FObjectInitializer& ObjectInitializer = FObjectInitializer::Get());
 	// Begin UCommandlet Interface
 	virtual int32 Main(const FString& Params) override;
 	// End UCommandlet Interface
diff --git a/Engine/Source/Editor/UnrealEd/Classes/Commandlets/DiffAssetsCommandlet.h b/Engine/Source/Editor/UnrealEd/Classes/Commandlets/DiffAssetsCommandlet.h
--- a/Engine/Source/Editor/UnrealEd/Classes/Commandlets/DiffAssetsCommandlet.h
+++ b/Engine/Source/Editor/UnrealEd/Classes/Commandlets/DiffAssetsCommandlet.h
@@ -11,7 +11,9 @@
 UCLASS()
 class UDiffAssetsCommandlet : public UCommandlet
 {
-    GENERATED_UCLASS_BODY()
+    GENERATED_BODY()
+public:
+    UDiffAssetsCommandlet(const FObjectInitializer& ObjectInitializer = FObjectInitializer::Get());
 	// Begin UCommandlet Interface
 	virtual int32 Main(const FString& Params) override
 	{
diff --git a/Engine/Source/Editor/UnrealEd/Classes/Commandlets/DiffPackagesCommandlet.h b/Engine/Source/Editor/UnrealEd/Classes/Commandlets/DiffPackagesCommandlet.h
--- a/Engine/Source/Editor/UnrealEd/Classes/Commandlets/DiffPackagesCommandlet.h
+++ b/Engine/Source/Editor/UnrealEd/Classes/Commandlets/DiffPackagesCommandlet.h
@@ -31,7 +31,9 @@ enum EObjectDiff
 UCLASS()
 class UDiffPackagesCommandlet : public UCommandlet
 {
-	GENERATED_UCLASS_BODY()
+	GENERATED_BODY()
+public:
+	UDiffPackagesCommandlet(const FObjectInitializer& ObjectInitializer = FObjectInitializer::Get());
 
 	UPROPERTY()
 	UPackage* Packages[3];
diff --git a/Engine/Source/Editor/UnrealEd/Classes/Commandlets/DumpBlueprintsInfoCommandlet.h b/Engine/Source/Editor/UnrealEd/Classes/Commandlets/DumpBlueprintsInfoCommandlet.h
--- a/Engine/Source/Editor/UnrealEd/Classes/Commandlets/DumpBlueprintsInfoCommandlet.h
+++ b/Engine/Source/Editor/UnrealEd/Classes/Commandlets/DumpBlueprintsInfoCommandlet.h
@@ -7,7 +7,9 @@
 UCLASS()
 class UDumpBlueprintsInfoCommandlet : public UCommandlet
 {
-	GENERATED_UCLASS_BODY()
+	GENERATED_BODY()
+public:
+	UDumpBlueprintsInfoCommandlet(const FObjectInitializer& ObjectInitializer = FObjectInitializer::Get());
 
 public:		
 	// Begin UCommandlet Interface
diff --git a/Engine/Source/Editor/UnrealEd/Classes/Commandlets/DumpHiddenCategoriesCommandlet.h b/Engine/Source/Editor/UnrealEd/Classes/Commandlets/DumpHiddenCategoriesCommandlet.h
--- a/Engine/Source/Editor/UnrealEd/Classes/Commandlets/DumpHiddenCategoriesCommandlet.h
+++ b/Engine/Source/Editor/UnrealEd/Classes/Commandlets/DumpHiddenCategoriesCommandlet.h
@@ -7,7 +7,9 @@
 UCLASS()
 class UDumpHiddenCategoriesCommandlet : public UCommandlet
 {
-	GENERATED_UCLASS_BODY()
+	GENERATED_BODY()
+public:
+	UDumpHiddenCategoriesCommandlet(const FObjectInitializer& ObjectInitializer = FObjectInitializer::Get());
 
 public:		
 	// Begin UCommandlet Interface
diff --git a/Engine/Source/Editor/UnrealEd/Classes/Commandlets/FileServerCommandlet.h b/Engine/Source/Editor/UnrealEd/Classes/Commandlets/FileServerCommandlet.h
--- a/Engine/Source/Editor/UnrealEd/Classes/Commandlets/FileServerCommandlet.h
+++ b/Engine/Source/Editor/UnrealEd/Classes/Commandlets/FileServerCommandlet.h
@@ -16,7 +16,9 @@ UCLASS()
 class UFileServerCommandlet
 	: public UCommandlet
 {
-	GENERATED_UCLASS_BODY()
+	GENERATED_BODY()
+public:
+	UFileServerCommandlet(const FObjectInitializer& ObjectInitializer = FObjectInitializer::Get());
 
 
 public:
@@ -32,4 +34,4 @@ private:
 
 	// Holds the instance identifier.
 	FGuid InstanceId;
-};
\ No newline at end of file
+};
diff --git a/Engine/Source/Editor/UnrealEd/Classes/Commandlets/FixupRedirectsCommandlet.h b/Engine/Source/Editor/UnrealEd/Classes/Commandlets/FixupRedirectsCommandlet.h
--- a/Engine/Source/Editor/UnrealEd/Classes/Commandlets/FixupRedirectsCommandlet.h
+++ b/Engine/Source/Editor/UnrealEd/Classes/Commandlets/FixupRedirectsCommandlet.h
@@ -7,7 +7,9 @@
 UCLASS()
 class UFixupRedirectsCommandlet : public UCommandlet
 {
-	GENERATED_UCLASS_BODY()
+	GENERATED_BODY()
+public:
+	UFixupRedirectsCommandlet(const FObjectInitializer& ObjectInitializer = FObjectInitializer::Get());
 	// Begin UCommandlet Interface
 	virtual void CreateCustomEngine(const FString& Params) override;
 	virtual int32 Main(const FString& Params) override;
diff --git a/Engine/Source/Editor/UnrealEd/Classes/Commandlets/GatherTextCommandlet.h b/Engine/Source/Editor/UnrealEd/Classes/Commandlets/GatherTextCommandlet.h
--- a/Engine/Source/Editor/UnrealEd/Classes/Commandlets/GatherTextCommandlet.h
+++ b/Engine/Source/Editor/UnrealEd/Classes/Commandlets/GatherTextCommandlet.h
@@ -18,7 +18,9 @@ namespace EOutputJson
 UCLASS()
 class UGatherTextCommandlet : public UGatherTextCommandletBase
 {
-    GENERATED_UCLASS_BODY()
+    GENERATED_BODY()
+public:
+    UGatherTextCommandlet(const FObjectInitializer& ObjectInitializer = FObjectInitializer::Get());
 public:
 	// Begin UCommandlet Interface
 	virtual int32 Main(const FString& Params) override;
diff --git a/Engine/Source/Editor/UnrealEd/Classes/Commandlets/GatherTextCommandletBase.h b/Engine/Source/Editor/UnrealEd/Classes/Commandlets/GatherTextCommandletBase.h
--- a/Engine/Source/Editor/UnrealEd/Classes/Commandlets/GatherTextCommandletBase.h
+++ b/Engine/Source/Editor/UnrealEd/Classes/Commandlets/GatherTextCommandletBase.h
@@ -248,7 +248,9 @@ class FJsonObject;
 UCLASS()
 class UGatherTextCommandletBase : public UCommandlet
 {
-	GENERATED_UCLASS_BODY()
+	GENERATED_BODY()
+public:
+	UGatherTextCommandletBase(const FObjectInitializer& ObjectInitializer = FObjectInitializer::Get());
 
 
 public:
diff --git a/Engine/Source/Editor/UnrealEd/Classes/Commandlets/GatherTextFromAssetsCommandlet.h b/Engine/Source/Editor/UnrealEd/Classes/Commandlets/GatherTextFromAssetsCommandlet.h
--- a/Engine/Source/Editor/UnrealEd/Classes/Commandlets/GatherTextFromAssetsCommandlet.h
+++ b/Engine/Source/Editor/UnrealEd/Classes/Commandlets/GatherTextFromAssetsCommandlet.h
@@ -30,7 +30,9 @@ class UDialogueWave;
 UCLASS()
 class UGatherTextFromAssetsCommandlet : public UGatherTextCommandletBase
 {
-	GENERATED_UCLASS_BODY()
+	GENERATED_BODY()
+public:
+	UGatherTextFromAssetsCommandlet(const FObjectInitializer& ObjectInitializer = FObjectInitializer::Get());
 
 	void ProcessPackages( const TArray< UPackage* >& PackagesToProcess );
 	void ProcessDialogueWave( const UDialogueWave* DialogueWave );
diff --git a/Engine/Source/Editor/UnrealEd/Classes/Commandlets/GatherTextFromMetadataCommandlet.h b/Engine/Source/Editor/UnrealEd/Classes/Commandlets/GatherTextFromMetadataCommandlet.h
--- a/Engine/Source/Editor/UnrealEd/Classes/Commandlets/GatherTextFromMetadataCommandlet.h
+++ b/Engine/Source/Editor/UnrealEd/Classes/Commandlets/GatherTextFromMetadataCommandlet.h
@@ -10,7 +10,9 @@
 UCLASS()
 class UGatherTextFromMetaDataCommandlet : public UGatherTextCommandletBase
 {
-	GENERATED_UCLASS_BODY()
+	GENERATED_BODY()
+public:
+	UGatherTextFromMetaDataCommandlet(const FObjectInitializer& ObjectInitializer = FObjectInitializer::Get());
 
 public:
 	// Begin UCommandlet Interface
@@ -28,4 +30,4 @@ private:
 private:
 	void GatherTextFromUObjects(const TArray<FString>& IncludePaths, const TArray<FString>& ExcludePaths, const FGatherParameters& Arguments);
 	void GatherTextFromUObject(UField* const Field, const FGatherParameters& Arguments);
-};
\ No newline at end of file
+};
diff --git a/Engine/Source/Editor/UnrealEd/Classes/Commandlets/GatherTextFromSourceCommandlet.h b/Engine/Source/Editor/UnrealEd/Classes/Commandlets/GatherTextFromSourceCommandlet.h
--- a/Engine/Source/Editor/UnrealEd/Classes/Commandlets/GatherTextFromSourceCommandlet.h
+++ b/Engine/Source/Editor/UnrealEd/Classes/Commandlets/GatherTextFromSourceCommandlet.h
@@ -10,7 +10,9 @@
 UCLASS()
 class UGatherTextFromSourceCommandlet : public UGatherTextCommandletBase
 {
-	GENERATED_UCLASS_BODY()
+	GENERATED_BODY()
+public:
+	UGatherTextFromSourceCommandlet(const FObjectInitializer& ObjectInitializer = FObjectInitializer::Get());
 
 private:
 #define LOC_DEFINE_REGION
diff --git a/Engine/Source/Editor/UnrealEd/Classes/Commandlets/GenerateBlueprintAPICommandlet.h b/Engine/Source/Editor/UnrealEd/Classes/Commandlets/GenerateBlueprintAPICommandlet.h
--- a/Engine/Source/Editor/UnrealEd/Classes/Commandlets/GenerateBlueprintAPICommandlet.h
+++ b/Engine/Source/Editor/UnrealEd/Classes/Commandlets/GenerateBlueprintAPICommandlet.h
@@ -7,7 +7,9 @@
 UCLASS()
 class UGenerateBlueprintAPICommandlet : public UCommandlet
 {
-	GENERATED_UCLASS_BODY()
+	GENERATED_BODY()
+public:
+	UGenerateBlueprintAPICommandlet(const FObjectInitializer& ObjectInitializer = FObjectInitializer::Get());
 
 public:		
 	// Begin UCommandlet Interface
diff --git a/Engine/Source/Editor/UnrealEd/Classes/Commandlets/GenerateDistillFileSetsCommandlet.h b/Engine/Source/Editor/UnrealEd/Classes/Commandlets/GenerateDistillFileSetsCommandlet.h
--- a/Engine/Source/Editor/UnrealEd/Classes/Commandlets/GenerateDistillFileSetsCommandlet.h
+++ b/Engine/Source/Editor/UnrealEd/Classes/Commandlets/GenerateDistillFileSetsCommandlet.h
@@ -11,7 +11,9 @@
 UCLASS()
 class UGenerateDistillFileSetsCommandlet : public UCommandlet
 {
-	GENERATED_UCLASS_BODY()
+	GENERATED_BODY()
+public:
+	UGenerateDistillFileSetsCommandlet(const FObjectInitializer& ObjectInitializer = FObjectInitializer::Get());
 
 	// Begin UCommandlet Interface
 	virtual int32 Main(const FString& Params) override;
diff --git a/Engine/Source/Editor/UnrealEd/Classes/Commandlets/GenerateGatherArchiveCommandlet.h b/Engine/Source/Editor/UnrealEd/Classes/Commandlets/GenerateGatherArchiveCommandlet.h
--- a/Engine/Source/Editor/UnrealEd/Classes/Commandlets/GenerateGatherArchiveCommandlet.h
+++ b/Engine/Source/Editor/UnrealEd/Classes/Commandlets/GenerateGatherArchiveCommandlet.h
@@ -14,7 +14,9 @@
 UCLASS()
 class UGenerateGatherArchiveCommandlet : public UGatherTextCommandletBase
 {
-    GENERATED_UCLASS_BODY()
+    GENERATED_BODY()
+public:
+    UGenerateGatherArchiveCommandlet(const FObjectInitializer& ObjectInitializer = FObjectInitializer::Get());
 #if CPP || UE_BUILD_DOCS
 public:
 	// Begin UCommandlet Interface
diff --git a/Engine/Source/Editor/UnrealEd/Classes/Commandlets/GenerateGatherManifestCommandlet.h b/Engine/Source/Editor/UnrealEd/Classes/Commandlets/GenerateGatherManifestCommandlet.h
--- a/Engine/Source/Editor/UnrealEd/Classes/Commandlets/GenerateGatherManifestCommandlet.h
+++ b/Engine/Source/Editor/UnrealEd/Classes/Commandlets/GenerateGatherManifestCommandlet.h
@@ -10,7 +10,9 @@
 UCLASS()
 class UGenerateGatherManifestCommandlet : public UGatherTextCommandletBase
 {
-    GENERATED_UCLASS_BODY()
+    GENERATED_BODY()
+public:
+    UGenerateGatherManifestCommandlet(const FObjectInitializer& ObjectInitializer = FObjectInitializer::Get());
 #if CPP || UE_BUILD_DOCS
 public:
 	// Begin UCommandlet Interface
diff --git a/Engine/Source/Editor/UnrealEd/Classes/Commandlets/GenerateTextLocalizationReportCommandlet.h b/Engine/Source/Editor/UnrealEd/Classes/Commandlets/GenerateTextLocalizationReportCommandlet.h
--- a/Engine/Source/Editor/UnrealEd/Classes/Commandlets/GenerateTextLocalizationReportCommandlet.h
+++ b/Engine/Source/Editor/UnrealEd/Classes/Commandlets/GenerateTextLocalizationReportCommandlet.h
@@ -9,7 +9,9 @@
 UCLASS()
 class UGenerateTextLocalizationReportCommandlet : public UGatherTextCommandletBase
 {
-	GENERATED_UCLASS_BODY()
+	GENERATED_BODY()
+public:
+	UGenerateTextLocalizationReportCommandlet(const FObjectInitializer& ObjectInitializer = FObjectInitializer::Get());
 
 
 
@@ -48,4 +50,4 @@ private:
 	FString SectionName;
 	FString CmdlineTimeStamp;
 
-};
\ No newline at end of file
+};
diff --git a/Engine/Source/Editor/UnrealEd/Classes/Commandlets/GenerateTextLocalizationResourceCommandlet.h b/Engine/Source/Editor/UnrealEd/Classes/Commandlets/GenerateTextLocalizationResourceCommandlet.h
--- a/Engine/Source/Editor/UnrealEd/Classes/Commandlets/GenerateTextLocalizationResourceCommandlet.h
+++ b/Engine/Source/Editor/UnrealEd/Classes/Commandlets/GenerateTextLocalizationResourceCommandlet.h
@@ -12,10 +12,12 @@
 UCLASS()
 class UGenerateTextLocalizationResourceCommandlet : public UGatherTextCommandletBase
 {
-	GENERATED_UCLASS_BODY()
+	GENERATED_BODY()
+public:
+	UGenerateTextLocalizationResourceCommandlet(const FObjectInitializer& ObjectInitializer = FObjectInitializer::Get());
 
 public:
 	// Begin UCommandlet Interface
 	virtual int32 Main(const FString& Params) override;
 	// End UCommandlet Interface
-};
\ No newline at end of file
+};
diff --git a/Engine/Source/Editor/UnrealEd/Classes/Commandlets/InternationalizationConditioningCommandlet.h b/Engine/Source/Editor/UnrealEd/Classes/Commandlets/InternationalizationConditioningCommandlet.h
--- a/Engine/Source/Editor/UnrealEd/Classes/Commandlets/InternationalizationConditioningCommandlet.h
+++ b/Engine/Source/Editor/UnrealEd/Classes/Commandlets/InternationalizationConditioningCommandlet.h
@@ -13,7 +13,9 @@
 UCLASS()
 class UInternationalizationConditioningCommandlet : public UGatherTextCommandletBase
 {
-    GENERATED_UCLASS_BODY()
+    GENERATED_BODY()
+public:
+    UInternationalizationConditioningCommandlet(const FObjectInitializer& ObjectInitializer = FObjectInitializer::Get());
 
 	
 	/**
@@ -194,4 +196,4 @@ private:
 	FString SectionName;
 	FConfigCacheIni LegacyLocalizationCacheIni;
 
-};
\ No newline at end of file
+};
diff --git a/Engine/Source/Editor/UnrealEd/Classes/Commandlets/InternationalizationExportCommandlet.h b/Engine/Source/Editor/UnrealEd/Classes/Commandlets/InternationalizationExportCommandlet.h
--- a/Engine/Source/Editor/UnrealEd/Classes/Commandlets/InternationalizationExportCommandlet.h
+++ b/Engine/Source/Editor/UnrealEd/Classes/Commandlets/InternationalizationExportCommandlet.h
@@ -381,7 +381,9 @@ private:
 UCLASS()
 class UInternationalizationExportCommandlet : public UGatherTextCommandletBase
 {
-    GENERATED_UCLASS_BODY()
+    GENERATED_BODY()
+public:
+    UInternationalizationExportCommandlet(const FObjectInitializer& ObjectInitializer = FObjectInitializer::Get());
 
 public:
 	// Begin UCommandlet Interface
@@ -395,4 +397,4 @@ private:
 	FString ConfigPath;
 	FString SectionName;
 
-};
\ No newline at end of file
+};
diff --git a/Engine/Source/Editor/UnrealEd/Classes/Commandlets/ListMaterialsUsedWithMeshEmittersCommandlet.h b/Engine/Source/Editor/UnrealEd/Classes/Commandlets/ListMaterialsUsedWithMeshEmittersCommandlet.h
--- a/Engine/Source/Editor/UnrealEd/Classes/Commandlets/ListMaterialsUsedWithMeshEmittersCommandlet.h
+++ b/Engine/Source/Editor/UnrealEd/Classes/Commandlets/ListMaterialsUsedWithMeshEmittersCommandlet.h
@@ -7,7 +7,9 @@
 UCLASS()
 class UListMaterialsUsedWithMeshEmittersCommandlet : public UCommandlet
 {
-	GENERATED_UCLASS_BODY()
+	GENERATED_BODY()
+public:
+	UListMaterialsUsedWithMeshEmittersCommandlet(const FObjectInitializer& ObjectInitializer = FObjectInitializer::Get());
 	
 	/** Process the given particle system. For mesh emitters, output material paths that don't have bUsedWithParticleSprites flagged.*/
 	void ProcessParticleSystem(UParticleSystem* ParticleSystem , TArray<FString> &OutMaterials);
@@ -15,4 +17,4 @@ class UListMaterialsUsedWithMeshEmittersCommandlet : public UCommandlet
 	// Begin UCommandlet Interface
 	virtual int32 Main(const FString& Params) override;
 	// End UCommandlet Interface
-};
\ No newline at end of file
+};
diff --git a/Engine/Source/Editor/UnrealEd/Classes/Commandlets/ListStaticMeshesImportedFromSpeedTreesCommandlet.h b/Engine/Source/Editor/UnrealEd/Classes/Commandlets/ListStaticMeshesImportedFromSpeedTreesCommandlet.h
--- a/Engine/Source/Editor/UnrealEd/Classes/Commandlets/ListStaticMeshesImportedFromSpeedTreesCommandlet.h
+++ b/Engine/Source/Editor/UnrealEd/Classes/Commandlets/ListStaticMeshesImportedFromSpeedTreesCommandlet.h
@@ -7,9 +7,11 @@
 UCLASS()
 class UListStaticMeshesImportedFromSpeedTreesCommandlet : public UCommandlet
 {
-	GENERATED_UCLASS_BODY()
+	GENERATED_BODY()
+public:
+	UListStaticMeshesImportedFromSpeedTreesCommandlet(const FObjectInitializer& ObjectInitializer = FObjectInitializer::Get());
 
 	// Begin UCommandlet Interface
 	virtual int32 Main(const FString& Params) override;
 	// End UCommandlet Interface
-};
\ No newline at end of file
+};
diff --git a/Engine/Source/Editor/UnrealEd/Classes/Commandlets/LoadPackageCommandlet.h b/Engine/Source/Editor/UnrealEd/Classes/Commandlets/LoadPackageCommandlet.h
--- a/Engine/Source/Editor/UnrealEd/Classes/Commandlets/LoadPackageCommandlet.h
+++ b/Engine/Source/Editor/UnrealEd/Classes/Commandlets/LoadPackageCommandlet.h
@@ -7,7 +7,9 @@
 UCLASS()
 class ULoadPackageCommandlet : public UCommandlet
 {
-	GENERATED_UCLASS_BODY()
+	GENERATED_BODY()
+public:
+	ULoadPackageCommandlet(const FObjectInitializer& ObjectInitializer = FObjectInitializer::Get());
 	/**
 	 *	Parse the given load list file, placing the entries in the given Tokens array.
 	 *
diff --git a/Engine/Source/Editor/UnrealEd/Classes/Commandlets/ParticleSystemAuditCommandlet.h b/Engine/Source/Editor/UnrealEd/Classes/Commandlets/ParticleSystemAuditCommandlet.h
--- a/Engine/Source/Editor/UnrealEd/Classes/Commandlets/ParticleSystemAuditCommandlet.h
+++ b/Engine/Source/Editor/UnrealEd/Classes/Commandlets/ParticleSystemAuditCommandlet.h
@@ -9,7 +9,9 @@
 UCLASS(config=Editor)
 class UParticleSystemAuditCommandlet : public UCommandlet
 {
-	GENERATED_UCLASS_BODY()
+	GENERATED_BODY()
+public:
+	UParticleSystemAuditCommandlet(const FObjectInitializer& ObjectInitializer = FObjectInitializer::Get());
 
 	/** Helper struct for tracking disabled LOD level issues */
 	struct FParticleSystemLODInfo
diff --git a/Engine/Source/Editor/UnrealEd/Classes/Commandlets/PkgInfoCommandlet.h b/Engine/Source/Editor/UnrealEd/Classes/Commandlets/PkgInfoCommandlet.h
--- a/Engine/Source/Editor/UnrealEd/Classes/Commandlets/PkgInfoCommandlet.h
+++ b/Engine/Source/Editor/UnrealEd/Classes/Commandlets/PkgInfoCommandlet.h
@@ -7,7 +7,9 @@
 UCLASS()
 class UPkgInfoCommandlet : public UCommandlet
 {
-	GENERATED_UCLASS_BODY()
+	GENERATED_BODY()
+public:
+	UPkgInfoCommandlet(const FObjectInitializer& ObjectInitializer = FObjectInitializer::Get());
 	// Begin UCommandlet Interface
 	virtual int32 Main(const FString& Params) override;
 	// End UCommandlet Interface
diff --git a/Engine/Source/Editor/UnrealEd/Classes/Commandlets/RepairLocalizationDataCommandlet.h b/Engine/Source/Editor/UnrealEd/Classes/Commandlets/RepairLocalizationDataCommandlet.h
--- a/Engine/Source/Editor/UnrealEd/Classes/Commandlets/RepairLocalizationDataCommandlet.h
+++ b/Engine/Source/Editor/UnrealEd/Classes/Commandlets/RepairLocalizationDataCommandlet.h
@@ -9,10 +9,12 @@
 UCLASS()
 class URepairLocalizationDataCommandlet : public UGatherTextCommandletBase
 {
-	GENERATED_UCLASS_BODY()
+	GENERATED_BODY()
+public:
+	URepairLocalizationDataCommandlet(const FObjectInitializer& ObjectInitializer = FObjectInitializer::Get());
 
 public:
 	// Begin UCommandlet Interface
 	virtual int32 Main(const FString& Params) override;
 	// End UCommandlet Interface
-};
\ No newline at end of file
+};
diff --git a/Engine/Source/Editor/UnrealEd/Classes/Commandlets/ReplaceActorCommandlet.h b/Engine/Source/Editor/UnrealEd/Classes/Commandlets/ReplaceActorCommandlet.h
--- a/Engine/Source/Editor/UnrealEd/Classes/Commandlets/ReplaceActorCommandlet.h
+++ b/Engine/Source/Editor/UnrealEd/Classes/Commandlets/ReplaceActorCommandlet.h
@@ -8,7 +8,9 @@
 UCLASS()
 class UReplaceActorCommandlet : public UCommandlet
 {
-	GENERATED_UCLASS_BODY()
+	GENERATED_BODY()
+public:
+	UReplaceActorCommandlet(const FObjectInitializer& ObjectInitializer = FObjectInitializer::Get());
 	// Begin UCommandlet Interface
 	virtual int32 Main(const FString& Params) override;
 	// End UCommandlet Interface
diff --git a/Engine/Source/Editor/UnrealEd/Classes/Commandlets/ResavePackagesCommandlet.h b/Engine/Source/Editor/UnrealEd/Classes/Commandlets/ResavePackagesCommandlet.h
--- a/Engine/Source/Editor/UnrealEd/Classes/Commandlets/ResavePackagesCommandlet.h
+++ b/Engine/Source/Editor/UnrealEd/Classes/Commandlets/ResavePackagesCommandlet.h
@@ -8,7 +8,9 @@ UCLASS()
 // Added UNREALED_API to expose this to the save packages test
 class UNREALED_API UResavePackagesCommandlet : public UCommandlet
 {
-    GENERATED_UCLASS_BODY()
+    GENERATED_BODY()
+public:
+    UResavePackagesCommandlet(const FObjectInitializer& ObjectInitializer = FObjectInitializer::Get());
 
 	enum EBrevity
 	{
diff --git a/Engine/Source/Editor/UnrealEd/Classes/Commandlets/UpdateGameProjectCommandlet.h b/Engine/Source/Editor/UnrealEd/Classes/Commandlets/UpdateGameProjectCommandlet.h
--- a/Engine/Source/Editor/UnrealEd/Classes/Commandlets/UpdateGameProjectCommandlet.h
+++ b/Engine/Source/Editor/UnrealEd/Classes/Commandlets/UpdateGameProjectCommandlet.h
@@ -7,7 +7,9 @@
 UCLASS()
 class UUpdateGameProjectCommandlet : public UCommandlet
 {
-	GENERATED_UCLASS_BODY()
+	GENERATED_BODY()
+public:
+	UUpdateGameProjectCommandlet(const FObjectInitializer& ObjectInitializer = FObjectInitializer::Get());
 
 	// Begin UCommandlet Interface
 	virtual int32 Main(const FString& Params) override;
diff --git a/Engine/Source/Editor/UnrealEd/Classes/Commandlets/WrangleContentCommandlet.h b/Engine/Source/Editor/UnrealEd/Classes/Commandlets/WrangleContentCommandlet.h
--- a/Engine/Source/Editor/UnrealEd/Classes/Commandlets/WrangleContentCommandlet.h
+++ b/Engine/Source/Editor/UnrealEd/Classes/Commandlets/WrangleContentCommandlet.h
@@ -7,7 +7,9 @@
 UCLASS()
 class UWrangleContentCommandlet : public UCommandlet
 {
-	GENERATED_UCLASS_BODY()
+	GENERATED_BODY()
+public:
+	UWrangleContentCommandlet(const FObjectInitializer& ObjectInitializer = FObjectInitializer::Get());
 	// Begin UCommandlet Interface
 	virtual int32 Main(const FString& Params) override;
 	// End UCommandlet Interface
diff --git a/Engine/Source/Editor/UnrealEd/Classes/CookOnTheSide/CookOnTheFlyServer.h b/Engine/Source/Editor/UnrealEd/Classes/CookOnTheSide/CookOnTheFlyServer.h
--- a/Engine/Source/Editor/UnrealEd/Classes/CookOnTheSide/CookOnTheFlyServer.h
+++ b/Engine/Source/Editor/UnrealEd/Classes/CookOnTheSide/CookOnTheFlyServer.h
@@ -52,7 +52,7 @@ UCLASS()
 class UNREALED_API UCookOnTheFlyServer : public UObject, public FTickableEditorObject
 {
 	GENERATED_BODY()
-
+public:
 	UCookOnTheFlyServer(const FObjectInitializer& ObjectInitializer = FObjectInitializer::Get());
 
 private:
diff --git a/Engine/Source/Editor/UnrealEd/Classes/Editor/EditorEngine.h b/Engine/Source/Editor/UnrealEd/Classes/Editor/EditorEngine.h
--- a/Engine/Source/Editor/UnrealEd/Classes/Editor/EditorEngine.h
+++ b/Engine/Source/Editor/UnrealEd/Classes/Editor/EditorEngine.h
@@ -42,7 +42,7 @@ enum EPasteTo
 USTRUCT()
 struct FSlatePlayInEditorInfo
 {
-	GENERATED_USTRUCT_BODY()
+	GENERATED_BODY()
 
 	/** The spawned player for updating viewport location from player when pie closes */
 	TWeakObjectPtr<class ULocalPlayer>	EditorPlayer;
@@ -67,9 +67,9 @@ struct FSlatePlayInEditorInfo
 USTRUCT()
 struct FPIELoginInfo
 {
+	GENERATED_BODY()
 public:
 
-	GENERATED_USTRUCT_BODY()
 
 	/** Type of account. Needed to identity the auth method to use (epic, internal, facebook, etc) */
 	UPROPERTY()
@@ -120,7 +120,7 @@ struct FPieLoginStruct
 USTRUCT()
 struct FCopySelectedInfo
 {
-	GENERATED_USTRUCT_BODY()
+	GENERATED_BODY()
 
 	/** Do not cache this info, it is only valid after a call to CanCopySelectedActorsToClipboard has been made, and becomes redundant
 	when the current selection changes. Used to determine whether a copy can be performed based on the current selection state */
diff --git a/Engine/Source/Editor/UnrealEd/Classes/Editor/EditorGameAgnosticSettings.h b/Engine/Source/Editor/UnrealEd/Classes/Editor/EditorGameAgnosticSettings.h
--- a/Engine/Source/Editor/UnrealEd/Classes/Editor/EditorGameAgnosticSettings.h
+++ b/Engine/Source/Editor/UnrealEd/Classes/Editor/EditorGameAgnosticSettings.h
@@ -7,7 +7,9 @@
 UCLASS(config=EditorGameAgnostic)
 class UEditorGameAgnosticSettings : public UObject
 {
-	GENERATED_UCLASS_BODY()
+	GENERATED_BODY()
+public:
+	UEditorGameAgnosticSettings(const FObjectInitializer& ObjectInitializer = FObjectInitializer::Get());
 
 	/** When checked, the most recently loaded project will be auto-loaded at editor startup if no other project was specified on the command line */
 	UPROPERTY(EditAnywhere, Category=General)
@@ -58,4 +60,4 @@ class UEditorGameAgnosticSettings : public UObject
 	// Begin UObject Interface
 	virtual void PostEditChangeProperty( struct FPropertyChangedEvent& PropertyChangedEvent) override;
 	// End UObject Interface
-};
\ No newline at end of file
+};
