Back out changelist 2481333

diff --git a/Engine/Source/Editor/BlueprintGraph/Classes/K2Node_StructMemberSet.h b/Engine/Source/Editor/BlueprintGraph/Classes/K2Node_StructMemberSet.h
--- a/Engine/Source/Editor/BlueprintGraph/Classes/K2Node_StructMemberSet.h
+++ b/Engine/Source/Editor/BlueprintGraph/Classes/K2Node_StructMemberSet.h
@@ -9,9 +9,7 @@
 UCLASS(MinimalAPI)
 class UK2Node_StructMemberSet : public UK2Node_StructOperation
 {
-	GENERATED_BODY()
-public:
-	BLUEPRINTGRAPH_API UK2Node_StructMemberSet(const FObjectInitializer& ObjectInitializer = FObjectInitializer::Get());
+	GENERATED_UCLASS_BODY()
 
 	UPROPERTY(EditAnywhere, Category=PinOptions, EditFixedSize)
 	TArray<FOptionalPinFromProperty> ShowPinForProperties;
diff --git a/Engine/Source/Editor/BlueprintGraph/Classes/K2Node_StructOperation.h b/Engine/Source/Editor/BlueprintGraph/Classes/K2Node_StructOperation.h
--- a/Engine/Source/Editor/BlueprintGraph/Classes/K2Node_StructOperation.h
+++ b/Engine/Source/Editor/BlueprintGraph/Classes/K2Node_StructOperation.h
@@ -8,9 +8,7 @@
 UCLASS(MinimalAPI, abstract)
 class UK2Node_StructOperation : public UK2Node_Variable
 {
-	GENERATED_BODY()
-public:
-	BLUEPRINTGRAPH_API UK2Node_StructOperation(const FObjectInitializer& ObjectInitializer = FObjectInitializer::Get());
+	GENERATED_UCLASS_BODY()
 
 	/** Class that this variable is defined in.  */
 	UPROPERTY()
diff --git a/Engine/Source/Editor/BlueprintGraph/Classes/K2Node_Switch.h b/Engine/Source/Editor/BlueprintGraph/Classes/K2Node_Switch.h
--- a/Engine/Source/Editor/BlueprintGraph/Classes/K2Node_Switch.h
+++ b/Engine/Source/Editor/BlueprintGraph/Classes/K2Node_Switch.h
@@ -8,9 +8,7 @@
 UCLASS(MinimalAPI, abstract)
 class UK2Node_Switch : public UK2Node
 {
-	GENERATED_BODY()
-public:
-	BLUEPRINTGRAPH_API UK2Node_Switch(const FObjectInitializer& ObjectInitializer = FObjectInitializer::Get());
+	GENERATED_UCLASS_BODY()
 
 	/** If true switch has a default pin */
 	UPROPERTY(EditAnywhere, Category=PinOptions)
diff --git a/Engine/Source/Editor/BlueprintGraph/Classes/K2Node_SwitchInteger.h b/Engine/Source/Editor/BlueprintGraph/Classes/K2Node_SwitchInteger.h
--- a/Engine/Source/Editor/BlueprintGraph/Classes/K2Node_SwitchInteger.h
+++ b/Engine/Source/Editor/BlueprintGraph/Classes/K2Node_SwitchInteger.h
@@ -8,9 +8,7 @@
 UCLASS(MinimalAPI)
 class UK2Node_SwitchInteger : public UK2Node_Switch
 {
-	GENERATED_BODY()
-public:
-	BLUEPRINTGRAPH_API UK2Node_SwitchInteger(const FObjectInitializer& ObjectInitializer = FObjectInitializer::Get());
+	GENERATED_UCLASS_BODY()
 
 	/** Set the starting index for the node */
 	UPROPERTY(EditAnywhere, Category=PinOptions, meta=(NoSpinbox=true))
diff --git a/Engine/Source/Editor/BlueprintGraph/Classes/K2Node_SwitchName.h b/Engine/Source/Editor/BlueprintGraph/Classes/K2Node_SwitchName.h
--- a/Engine/Source/Editor/BlueprintGraph/Classes/K2Node_SwitchName.h
+++ b/Engine/Source/Editor/BlueprintGraph/Classes/K2Node_SwitchName.h
@@ -8,9 +8,7 @@
 UCLASS(MinimalAPI)
 class UK2Node_SwitchName : public UK2Node_Switch
 {
-	GENERATED_BODY()
-public:
-	BLUEPRINTGRAPH_API UK2Node_SwitchName(const FObjectInitializer& ObjectInitializer = FObjectInitializer::Get());
+	GENERATED_UCLASS_BODY()
 
 	UPROPERTY(EditAnywhere, Category=PinOptions)
 	TArray<FName> PinNames;
diff --git a/Engine/Source/Editor/BlueprintGraph/Classes/K2Node_SwitchString.h b/Engine/Source/Editor/BlueprintGraph/Classes/K2Node_SwitchString.h
--- a/Engine/Source/Editor/BlueprintGraph/Classes/K2Node_SwitchString.h
+++ b/Engine/Source/Editor/BlueprintGraph/Classes/K2Node_SwitchString.h
@@ -8,9 +8,7 @@
 UCLASS(MinimalAPI)
 class UK2Node_SwitchString : public UK2Node_Switch
 {
-	GENERATED_BODY()
-public:
-	BLUEPRINTGRAPH_API UK2Node_SwitchString(const FObjectInitializer& ObjectInitializer = FObjectInitializer::Get());
+	GENERATED_UCLASS_BODY()
 
 	UPROPERTY(EditAnywhere, Category=PinOptions)
 	TArray<FString> PinNames;
diff --git a/Engine/Source/Editor/BlueprintGraph/Classes/K2Node_Timeline.h b/Engine/Source/Editor/BlueprintGraph/Classes/K2Node_Timeline.h
--- a/Engine/Source/Editor/BlueprintGraph/Classes/K2Node_Timeline.h
+++ b/Engine/Source/Editor/BlueprintGraph/Classes/K2Node_Timeline.h
@@ -8,9 +8,7 @@
 UCLASS(MinimalAPI)
 class UK2Node_Timeline : public UK2Node
 {
-	GENERATED_BODY()
-public:
-	BLUEPRINTGRAPH_API UK2Node_Timeline(const FObjectInitializer& ObjectInitializer = FObjectInitializer::Get());
+	GENERATED_UCLASS_BODY()
 
 	/** The name of the timeline. Used to name ONLY the member variable (Component). To obtain the name of timeline template use UTimelineTemplate::TimelineVariableNameToTemplateName */
 	UPROPERTY()
diff --git a/Engine/Source/Editor/BlueprintGraph/Classes/K2Node_Tunnel.h b/Engine/Source/Editor/BlueprintGraph/Classes/K2Node_Tunnel.h
--- a/Engine/Source/Editor/BlueprintGraph/Classes/K2Node_Tunnel.h
+++ b/Engine/Source/Editor/BlueprintGraph/Classes/K2Node_Tunnel.h
@@ -7,9 +7,7 @@
 UCLASS(MinimalAPI)
 class UK2Node_Tunnel : public UK2Node_EditablePinBase
 {
-	GENERATED_BODY()
-public:
-	BLUEPRINTGRAPH_API UK2Node_Tunnel(const FObjectInitializer& ObjectInitializer = FObjectInitializer::Get());
+	GENERATED_UCLASS_BODY()
 
 	// A tunnel node either has output pins that came from another tunnel's input pins, or vis versa
 	// Note: OutputSourceNode might be equal to InputSinkNode
diff --git a/Engine/Source/Editor/BlueprintGraph/Classes/K2Node_Variable.h b/Engine/Source/Editor/BlueprintGraph/Classes/K2Node_Variable.h
--- a/Engine/Source/Editor/BlueprintGraph/Classes/K2Node_Variable.h
+++ b/Engine/Source/Editor/BlueprintGraph/Classes/K2Node_Variable.h
@@ -22,9 +22,7 @@ namespace ESelfContextInfo
 UCLASS(abstract)
 class BLUEPRINTGRAPH_API UK2Node_Variable : public UK2Node
 {
-	GENERATED_BODY()
-public:
-	UK2Node_Variable(const FObjectInitializer& ObjectInitializer = FObjectInitializer::Get());
+	GENERATED_UCLASS_BODY()
 
 	/** Reference to variable we want to set/get */
 	UPROPERTY()
diff --git a/Engine/Source/Editor/BlueprintGraph/Classes/K2Node_VariableGet.h b/Engine/Source/Editor/BlueprintGraph/Classes/K2Node_VariableGet.h
--- a/Engine/Source/Editor/BlueprintGraph/Classes/K2Node_VariableGet.h
+++ b/Engine/Source/Editor/BlueprintGraph/Classes/K2Node_VariableGet.h
@@ -9,9 +9,7 @@
 UCLASS()
 class BLUEPRINTGRAPH_API UK2Node_VariableGet : public UK2Node_Variable
 {
-	GENERATED_BODY()
-public:
-	UK2Node_VariableGet(const FObjectInitializer& ObjectInitializer = FObjectInitializer::Get());
+	GENERATED_UCLASS_BODY()
 
 	// Begin UEdGraphNode interface
 	virtual void AllocateDefaultPins() override;
diff --git a/Engine/Source/Editor/BlueprintGraph/Classes/NodeDependingOnEnumInterface.h b/Engine/Source/Editor/BlueprintGraph/Classes/NodeDependingOnEnumInterface.h
--- a/Engine/Source/Editor/BlueprintGraph/Classes/NodeDependingOnEnumInterface.h
+++ b/Engine/Source/Editor/BlueprintGraph/Classes/NodeDependingOnEnumInterface.h
@@ -6,15 +6,12 @@
 UINTERFACE(meta=(CannotImplementInterfaceInBlueprint))
 class BLUEPRINTGRAPH_API UNodeDependingOnEnumInterface : public UInterface
 {
-	GENERATED_BODY()
-public:
-	UNodeDependingOnEnumInterface(const FObjectInitializer& ObjectInitializer = FObjectInitializer::Get());
+	GENERATED_UINTERFACE_BODY()
 };
 
 class BLUEPRINTGRAPH_API INodeDependingOnEnumInterface
 {
-	GENERATED_BODY()
-public:
+	GENERATED_IINTERFACE_BODY()
 
 	virtual class UEnum* GetEnum() const PURE_VIRTUAL(INodeDependingOnEnumInterface::GetEnum,return NULL;);
 
diff --git a/Engine/Source/Editor/BlueprintGraph/Public/BlueprintEditorSettings.h b/Engine/Source/Editor/BlueprintGraph/Public/BlueprintEditorSettings.h
--- a/Engine/Source/Editor/BlueprintGraph/Public/BlueprintEditorSettings.h
+++ b/Engine/Source/Editor/BlueprintGraph/Public/BlueprintEditorSettings.h
@@ -17,9 +17,7 @@ UCLASS(config=EditorUserSettings)
 class BLUEPRINTGRAPH_API UBlueprintEditorSettings
 	:	public UObject
 {
-	GENERATED_BODY()
-public:
-	UBlueprintEditorSettings(const FObjectInitializer& ObjectInitializer = FObjectInitializer::Get());
+	GENERATED_UCLASS_BODY()
 
 // Style Settings
 public:
diff --git a/Engine/Source/Editor/BlueprintGraph/Public/BlueprintNodeSpawner.h b/Engine/Source/Editor/BlueprintGraph/Public/BlueprintNodeSpawner.h
--- a/Engine/Source/Editor/BlueprintGraph/Public/BlueprintNodeSpawner.h
+++ b/Engine/Source/Editor/BlueprintGraph/Public/BlueprintNodeSpawner.h
@@ -42,9 +42,7 @@ struct FBlueprintActionUiSpec
 UCLASS(Transient)
 class BLUEPRINTGRAPH_API UBlueprintNodeSpawner : public UObject, public IBlueprintNodeBinder
 {
-	GENERATED_BODY()
-public:
-	UBlueprintNodeSpawner(const FObjectInitializer& ObjectInitializer = FObjectInitializer::Get());
+	GENERATED_UCLASS_BODY()
 	DECLARE_DELEGATE_TwoParams(FCustomizeNodeDelegate, UEdGraphNode*, bool);
 	DECLARE_DELEGATE_ThreeParams(FUiSpecOverrideDelegate, FBlueprintActionContext const&, FBindingSet const&, FBlueprintActionUiSpec*);
 
diff --git a/Engine/Source/Editor/Blutility/Classes/EditorUtilityBlueprint.h b/Engine/Source/Editor/Blutility/Classes/EditorUtilityBlueprint.h
--- a/Engine/Source/Editor/Blutility/Classes/EditorUtilityBlueprint.h
+++ b/Engine/Source/Editor/Blutility/Classes/EditorUtilityBlueprint.h
@@ -11,9 +11,7 @@
 UCLASS()
 class UEditorUtilityBlueprint : public UBlueprint
 {
-	GENERATED_BODY()
-public:
-	UEditorUtilityBlueprint(const FObjectInitializer& ObjectInitializer = FObjectInitializer::Get());
+	GENERATED_UCLASS_BODY()
 
 #if WITH_EDITOR
 	// UBlueprint interface
diff --git a/Engine/Source/Editor/Blutility/Classes/EditorUtilityBlueprintFactory.h b/Engine/Source/Editor/Blutility/Classes/EditorUtilityBlueprintFactory.h
--- a/Engine/Source/Editor/Blutility/Classes/EditorUtilityBlueprintFactory.h
+++ b/Engine/Source/Editor/Blutility/Classes/EditorUtilityBlueprintFactory.h
@@ -11,9 +11,7 @@
 UCLASS()
 class UEditorUtilityBlueprintFactory : public UFactory
 {
-	GENERATED_BODY()
-public:
-	UEditorUtilityBlueprintFactory(const FObjectInitializer& ObjectInitializer = FObjectInitializer::Get());
+	GENERATED_UCLASS_BODY()
 
 	// The parent class of the created blueprint
 	UPROPERTY(EditAnywhere, Category=BlueprintFactory, meta=(AllowAbstract = "", BlueprintBaseOnly = ""))
diff --git a/Engine/Source/Editor/Blutility/Classes/GlobalEditorUtilityBase.h b/Engine/Source/Editor/Blutility/Classes/GlobalEditorUtilityBase.h
--- a/Engine/Source/Editor/Blutility/Classes/GlobalEditorUtilityBase.h
+++ b/Engine/Source/Editor/Blutility/Classes/GlobalEditorUtilityBase.h
@@ -16,9 +16,7 @@ DECLARE_DYNAMIC_MULTICAST_DELEGATE_TwoParams( FForEachAssetIteratorSignature, cl
 UCLASS(Abstract, hideCategories=(Object), Blueprintable)
 class BLUTILITY_API UGlobalEditorUtilityBase : public UObject
 {
-	GENERATED_BODY()
-public:
-	UGlobalEditorUtilityBase(const FObjectInitializer& ObjectInitializer = FObjectInitializer::Get());
+	GENERATED_UCLASS_BODY()
 
 	UPROPERTY(Category=Config, EditDefaultsOnly, BlueprintReadWrite, AssetRegistrySearchable)
 	FString HelpText;
diff --git a/Engine/Source/Editor/Blutility/Classes/PlacedEditorUtilityBase.h b/Engine/Source/Editor/Blutility/Classes/PlacedEditorUtilityBase.h
--- a/Engine/Source/Editor/Blutility/Classes/PlacedEditorUtilityBase.h
+++ b/Engine/Source/Editor/Blutility/Classes/PlacedEditorUtilityBase.h
@@ -11,9 +11,7 @@
 UCLASS(Abstract, hideCategories=(Object, Actor)/*, Blueprintable*/)
 class BLUTILITY_API APlacedEditorUtilityBase : public AActor
 {
-	GENERATED_BODY()
-public:
-	APlacedEditorUtilityBase(const FObjectInitializer& ObjectInitializer = FObjectInitializer::Get());
+	GENERATED_UCLASS_BODY()
 
 	UPROPERTY(Category=Config, EditDefaultsOnly, BlueprintReadWrite)
 	FString HelpText;
diff --git a/Engine/Source/Editor/Cascade/Classes/CascadeConfiguration.h b/Engine/Source/Editor/Cascade/Classes/CascadeConfiguration.h
--- a/Engine/Source/Editor/Cascade/Classes/CascadeConfiguration.h
+++ b/Engine/Source/Editor/Cascade/Classes/CascadeConfiguration.h
@@ -13,7 +13,7 @@
 USTRUCT()
 struct FModuleMenuMapper
 {
-	GENERATED_BODY()
+	GENERATED_USTRUCT_BODY()
 
 	UPROPERTY()
 	FString ObjName;
@@ -26,9 +26,7 @@ struct FModuleMenuMapper
 UCLASS(hidecategories=Object, config=Editor)
 class UCascadeConfiguration : public UObject
 {
-	GENERATED_BODY()
-public:
-	UCascadeConfiguration(const FObjectInitializer& ObjectInitializer = FObjectInitializer::Get());
+	GENERATED_UCLASS_BODY()
 
 	/**
 	 *	TypeData-to-base module mappings.
diff --git a/Engine/Source/Editor/Cascade/Classes/CascadeParticleSystemComponent.h b/Engine/Source/Editor/Cascade/Classes/CascadeParticleSystemComponent.h
--- a/Engine/Source/Editor/Cascade/Classes/CascadeParticleSystemComponent.h
+++ b/Engine/Source/Editor/Cascade/Classes/CascadeParticleSystemComponent.h
@@ -11,9 +11,7 @@
 UCLASS(hidecategories=Object, hidecategories=Physics, hidecategories=Collision, editinlinenew)
 class UCascadeParticleSystemComponent : public UParticleSystemComponent
 {
-	GENERATED_BODY()
-public:
-	UCascadeParticleSystemComponent(const FObjectInitializer& ObjectInitializer = FObjectInitializer::Get());
+	GENERATED_UCLASS_BODY()
 
 
 public:	
diff --git a/Engine/Source/Editor/DestructibleMeshEditor/Classes/DestructibleChunkParamsProxy.h b/Engine/Source/Editor/DestructibleMeshEditor/Classes/DestructibleChunkParamsProxy.h
--- a/Engine/Source/Editor/DestructibleMeshEditor/Classes/DestructibleChunkParamsProxy.h
+++ b/Engine/Source/Editor/DestructibleMeshEditor/Classes/DestructibleChunkParamsProxy.h
@@ -9,9 +9,7 @@
 UCLASS()
 class UDestructibleChunkParamsProxy : public UObject
 {
-	GENERATED_BODY()
-public:
-	UDestructibleChunkParamsProxy(const FObjectInitializer& ObjectInitializer = FObjectInitializer::Get());
+	GENERATED_UCLASS_BODY()
 
 	UPROPERTY()
 	UDestructibleMesh* DestructibleMesh;
@@ -29,4 +27,4 @@ public:
 
 	TWeakPtr<IDestructibleMeshEditor> DestructibleMeshEditorPtr;
 #endif
-};
+};
\ No newline at end of file
diff --git a/Engine/Source/Editor/EditorLiveStreaming/EditorLiveStreamingSettings.h b/Engine/Source/Editor/EditorLiveStreaming/EditorLiveStreamingSettings.h
--- a/Engine/Source/Editor/EditorLiveStreaming/EditorLiveStreamingSettings.h
+++ b/Engine/Source/Editor/EditorLiveStreaming/EditorLiveStreamingSettings.h
@@ -38,9 +38,7 @@ namespace EEditorLiveStreamingWebCamResolution
 UCLASS(config=EditorGameAgnostic)
 class UEditorLiveStreamingSettings : public UObject
 {
-	GENERATED_BODY()
-public:
-	UEditorLiveStreamingSettings(const FObjectInitializer& ObjectInitializer = FObjectInitializer::Get());
+	GENERATED_UCLASS_BODY()
 
 public:
 
diff --git a/Engine/Source/Editor/EditorStyle/Public/Classes/EditorStyleSettings.h b/Engine/Source/Editor/EditorStyle/Public/Classes/EditorStyleSettings.h
--- a/Engine/Source/Editor/EditorStyle/Public/Classes/EditorStyleSettings.h
+++ b/Engine/Source/Editor/EditorStyle/Public/Classes/EditorStyleSettings.h
@@ -30,11 +30,9 @@ UCLASS(config=EditorUserSettings)
 class EDITORSTYLE_API UEditorStyleSettings
 	: public UObject
 {
-	GENERATED_BODY()
-public:
-	UEditorStyleSettings(const FObjectInitializer& ObjectInitializer = FObjectInitializer::Get());
 public:
 
+	GENERATED_UCLASS_BODY()
 
 public:
 
diff --git a/Engine/Source/Editor/EnvironmentQueryEditor/Classes/EdGraphSchema_EnvironmentQuery.h b/Engine/Source/Editor/EnvironmentQueryEditor/Classes/EdGraphSchema_EnvironmentQuery.h
--- a/Engine/Source/Editor/EnvironmentQueryEditor/Classes/EdGraphSchema_EnvironmentQuery.h
+++ b/Engine/Source/Editor/EnvironmentQueryEditor/Classes/EdGraphSchema_EnvironmentQuery.h
@@ -8,9 +8,7 @@
 UCLASS(MinimalAPI)
 class UEdGraphSchema_EnvironmentQuery : public UAIGraphSchema
 {
-	GENERATED_BODY()
-public:
-	ENVIRONMENTQUERYEDITOR_API UEdGraphSchema_EnvironmentQuery(const FObjectInitializer& ObjectInitializer = FObjectInitializer::Get());
+	GENERATED_UCLASS_BODY()
 
 	// Begin EdGraphSchema interface
 	virtual void CreateDefaultNodesForGraph(UEdGraph& Graph) const override;
diff --git a/Engine/Source/Editor/EnvironmentQueryEditor/Classes/EnvironmentQueryFactory.h b/Engine/Source/Editor/EnvironmentQueryEditor/Classes/EnvironmentQueryFactory.h
--- a/Engine/Source/Editor/EnvironmentQueryEditor/Classes/EnvironmentQueryFactory.h
+++ b/Engine/Source/Editor/EnvironmentQueryEditor/Classes/EnvironmentQueryFactory.h
@@ -6,9 +6,7 @@
 UCLASS(hidecategories=Object)
 class UEnvironmentQueryFactory : public UFactory
 {
-	GENERATED_BODY()
-public:
-	UEnvironmentQueryFactory(const FObjectInitializer& ObjectInitializer = FObjectInitializer::Get());
+	GENERATED_UCLASS_BODY()
 
 	// UFactory interface
 	virtual UObject* FactoryCreateNew(UClass* Class,UObject* InParent,FName Name,EObjectFlags Flags,UObject* Context,FFeedbackContext* Warn) override;
diff --git a/Engine/Source/Editor/EnvironmentQueryEditor/Classes/EnvironmentQueryGraph.h b/Engine/Source/Editor/EnvironmentQueryEditor/Classes/EnvironmentQueryGraph.h
--- a/Engine/Source/Editor/EnvironmentQueryEditor/Classes/EnvironmentQueryGraph.h
+++ b/Engine/Source/Editor/EnvironmentQueryEditor/Classes/EnvironmentQueryGraph.h
@@ -11,9 +11,7 @@ class UEnvironmentQueryGraphNode;
 UCLASS()
 class UEnvironmentQueryGraph : public UAIGraph
 {
-	GENERATED_BODY()
-public:
-	UEnvironmentQueryGraph(const FObjectInitializer& ObjectInitializer = FObjectInitializer::Get());
+	GENERATED_UCLASS_BODY()
 
 	virtual void Initialize() override;
 	virtual void OnLoaded() override;
diff --git a/Engine/Source/Editor/EnvironmentQueryEditor/Classes/EnvironmentQueryGraphNode.h b/Engine/Source/Editor/EnvironmentQueryEditor/Classes/EnvironmentQueryGraphNode.h
--- a/Engine/Source/Editor/EnvironmentQueryEditor/Classes/EnvironmentQueryGraphNode.h
+++ b/Engine/Source/Editor/EnvironmentQueryEditor/Classes/EnvironmentQueryGraphNode.h
@@ -8,9 +8,7 @@
 UCLASS()
 class UEnvironmentQueryGraphNode : public UAIGraphNode
 {
-	GENERATED_BODY()
-public:
-	UEnvironmentQueryGraphNode(const FObjectInitializer& ObjectInitializer = FObjectInitializer::Get());
+	GENERATED_UCLASS_BODY()
 
 	virtual class UEnvironmentQueryGraph* GetEnvironmentQueryGraph();
 
diff --git a/Engine/Source/Editor/EnvironmentQueryEditor/Classes/EnvironmentQueryGraphNode_Option.h b/Engine/Source/Editor/EnvironmentQueryEditor/Classes/EnvironmentQueryGraphNode_Option.h
--- a/Engine/Source/Editor/EnvironmentQueryEditor/Classes/EnvironmentQueryGraphNode_Option.h
+++ b/Engine/Source/Editor/EnvironmentQueryEditor/Classes/EnvironmentQueryGraphNode_Option.h
@@ -6,9 +6,7 @@
 UCLASS()
 class UEnvironmentQueryGraphNode_Option : public UEnvironmentQueryGraphNode
 {
-	GENERATED_BODY()
-public:
-	UEnvironmentQueryGraphNode_Option(const FObjectInitializer& ObjectInitializer = FObjectInitializer::Get());
+	GENERATED_UCLASS_BODY()
 
 	virtual void AllocateDefaultPins() override;
 	virtual void PostPlacedNewNode() override;
diff --git a/Engine/Source/Editor/EnvironmentQueryEditor/Classes/EnvironmentQueryGraphNode_Test.h b/Engine/Source/Editor/EnvironmentQueryEditor/Classes/EnvironmentQueryGraphNode_Test.h
--- a/Engine/Source/Editor/EnvironmentQueryEditor/Classes/EnvironmentQueryGraphNode_Test.h
+++ b/Engine/Source/Editor/EnvironmentQueryEditor/Classes/EnvironmentQueryGraphNode_Test.h
@@ -6,9 +6,7 @@
 UCLASS()
 class UEnvironmentQueryGraphNode_Test : public UEnvironmentQueryGraphNode
 {
-	GENERATED_BODY()
-public:
-	UEnvironmentQueryGraphNode_Test(const FObjectInitializer& ObjectInitializer = FObjectInitializer::Get());
+	GENERATED_UCLASS_BODY()
 
 	/** weight percent for display */
 	UPROPERTY()
diff --git a/Engine/Source/Editor/FoliageEdit/Private/ActorFactoryProceduralFoliage.h b/Engine/Source/Editor/FoliageEdit/Private/ActorFactoryProceduralFoliage.h
--- a/Engine/Source/Editor/FoliageEdit/Private/ActorFactoryProceduralFoliage.h
+++ b/Engine/Source/Editor/FoliageEdit/Private/ActorFactoryProceduralFoliage.h
@@ -8,9 +8,7 @@
 UCLASS(MinimalAPI, config=Editor)
 class UActorFactoryProceduralFoliage : public UActorFactoryBoxVolume
 {
-	GENERATED_BODY()
-public:
-	FOLIAGEEDIT_API UActorFactoryProceduralFoliage(const FObjectInitializer& ObjectInitializer = FObjectInitializer::Get());
+	GENERATED_UCLASS_BODY()
 
 	// Begin UActorFactory Interface
 	virtual bool CanCreateActorFrom( const FAssetData& AssetData, FText& OutErrorMsg ) override;
diff --git a/Engine/Source/Editor/FoliageEdit/Private/LandscapeGrassTypeFactory.h b/Engine/Source/Editor/FoliageEdit/Private/LandscapeGrassTypeFactory.h
--- a/Engine/Source/Editor/FoliageEdit/Private/LandscapeGrassTypeFactory.h
+++ b/Engine/Source/Editor/FoliageEdit/Private/LandscapeGrassTypeFactory.h
@@ -12,9 +12,7 @@
 UCLASS()
 class ULandscapeGrassTypeFactory : public UFactory
 {
-	GENERATED_BODY()
-public:
-	ULandscapeGrassTypeFactory(const FObjectInitializer& ObjectInitializer = FObjectInitializer::Get());
+	GENERATED_UCLASS_BODY()
 
 	// UFactory interface
 	virtual UObject* FactoryCreateNew(UClass* Class, UObject* InParent, FName Name, EObjectFlags Flags, UObject* Context, FFeedbackContext* Warn) override;
diff --git a/Engine/Source/Editor/FoliageEdit/Private/ProceduralFoliageFactory.h b/Engine/Source/Editor/FoliageEdit/Private/ProceduralFoliageFactory.h
--- a/Engine/Source/Editor/FoliageEdit/Private/ProceduralFoliageFactory.h
+++ b/Engine/Source/Editor/FoliageEdit/Private/ProceduralFoliageFactory.h
@@ -12,9 +12,7 @@
 UCLASS()
 class UProceduralFoliageFactory : public UFactory
 {
-	GENERATED_BODY()
-public:
-	UProceduralFoliageFactory(const FObjectInitializer& ObjectInitializer = FObjectInitializer::Get());
+	GENERATED_UCLASS_BODY()
 
 	// UFactory interface
 	virtual UObject* FactoryCreateNew(UClass* Class, UObject* InParent, FName Name, EObjectFlags Flags, UObject* Context, FFeedbackContext* Warn) override;
diff --git a/Engine/Source/Editor/GameProjectGeneration/Classes/DefaultTemplateProjectDefs.h b/Engine/Source/Editor/GameProjectGeneration/Classes/DefaultTemplateProjectDefs.h
--- a/Engine/Source/Editor/GameProjectGeneration/Classes/DefaultTemplateProjectDefs.h
+++ b/Engine/Source/Editor/GameProjectGeneration/Classes/DefaultTemplateProjectDefs.h
@@ -8,9 +8,7 @@
 UCLASS()
 class GAMEPROJECTGENERATION_API UDefaultTemplateProjectDefs : public UTemplateProjectDefs
 {
-	GENERATED_BODY()
-public:
-	UDefaultTemplateProjectDefs(const FObjectInitializer& ObjectInitializer = FObjectInitializer::Get());
+	GENERATED_UCLASS_BODY()
 
 	virtual bool GeneratesCode(const FString& ProjectTemplatePath) const override;
 
diff --git a/Engine/Source/Editor/GameProjectGeneration/Classes/TemplateProjectDefs.h b/Engine/Source/Editor/GameProjectGeneration/Classes/TemplateProjectDefs.h
--- a/Engine/Source/Editor/GameProjectGeneration/Classes/TemplateProjectDefs.h
+++ b/Engine/Source/Editor/GameProjectGeneration/Classes/TemplateProjectDefs.h
@@ -19,7 +19,7 @@ struct FTemplateConfigValue
 USTRUCT()
 struct FTemplateReplacement
 {
-	GENERATED_BODY()
+	GENERATED_USTRUCT_BODY()
 
 	UPROPERTY()
 	TArray<FString> Extensions;
@@ -37,7 +37,7 @@ struct FTemplateReplacement
 USTRUCT()
 struct FTemplateFolderRename
 {
-	GENERATED_BODY()
+	GENERATED_USTRUCT_BODY()
 
 	UPROPERTY()
 	FString From;
@@ -49,7 +49,7 @@ struct FTemplateFolderRename
 USTRUCT()
 struct FLocalizedTemplateString
 {
-	GENERATED_BODY()
+	GENERATED_USTRUCT_BODY()
 
 	UPROPERTY()
 	FString Language;
@@ -61,9 +61,7 @@ struct FLocalizedTemplateString
 UCLASS(abstract,config=TemplateDefs,MinimalAPI)
 class UTemplateProjectDefs : public UObject
 {
-	GENERATED_BODY()
-public:
-	GAMEPROJECTGENERATION_API UTemplateProjectDefs(const FObjectInitializer& ObjectInitializer = FObjectInitializer::Get());
+	GENERATED_UCLASS_BODY()
 
 	UPROPERTY(config)
 	TArray<FLocalizedTemplateString> LocalizedDisplayNames;
diff --git a/Engine/Source/Editor/GameplayAbilitiesEditor/Classes/GameplayAbilitiesBlueprintFactory.h b/Engine/Source/Editor/GameplayAbilitiesEditor/Classes/GameplayAbilitiesBlueprintFactory.h
--- a/Engine/Source/Editor/GameplayAbilitiesEditor/Classes/GameplayAbilitiesBlueprintFactory.h
+++ b/Engine/Source/Editor/GameplayAbilitiesEditor/Classes/GameplayAbilitiesBlueprintFactory.h
@@ -7,9 +7,7 @@
 UCLASS(HideCategories=Object, MinimalAPI)
 class UGameplayAbilitiesBlueprintFactory : public UFactory
 {
-	GENERATED_BODY()
-public:
-	GAMEPLAYABILITIESEDITOR_API UGameplayAbilitiesBlueprintFactory(const FObjectInitializer& ObjectInitializer = FObjectInitializer::Get());
+	GENERATED_UCLASS_BODY()
 
 	// The type of blueprint that will be created
 	UPROPERTY(EditAnywhere, Category=GameplayAbilitiesBlueprintFactory)
diff --git a/Engine/Source/Editor/GameplayAbilitiesEditor/Classes/GameplayAbilityGraphSchema.h b/Engine/Source/Editor/GameplayAbilitiesEditor/Classes/GameplayAbilityGraphSchema.h
--- a/Engine/Source/Editor/GameplayAbilitiesEditor/Classes/GameplayAbilityGraphSchema.h
+++ b/Engine/Source/Editor/GameplayAbilitiesEditor/Classes/GameplayAbilityGraphSchema.h
@@ -8,9 +8,7 @@
 UCLASS(MinimalAPI)
 class UGameplayAbilityGraphSchema : public UEdGraphSchema_K2
 {
-	GENERATED_BODY()
-public:
-	GAMEPLAYABILITIESEDITOR_API UGameplayAbilityGraphSchema(const FObjectInitializer& ObjectInitializer = FObjectInitializer::Get());
+	GENERATED_UCLASS_BODY()
 
 	/**
 	* Creates a new variable getter node and adds it to ParentGraph
@@ -35,4 +33,4 @@ public:
 	virtual class UK2Node_VariableSet* SpawnVariableSetNode(const FVector2D GraphPosition, class UEdGraph* ParentGraph, FName VariableName, UStruct* Source) const override;
 
 	virtual bool ShouldAlwaysPurgeOnModification() const override { return true; }
-};
+};
\ No newline at end of file
diff --git a/Engine/Source/Editor/GameplayAbilitiesEditor/Classes/K2Node_GameplayCueEvent.h b/Engine/Source/Editor/GameplayAbilitiesEditor/Classes/K2Node_GameplayCueEvent.h
--- a/Engine/Source/Editor/GameplayAbilitiesEditor/Classes/K2Node_GameplayCueEvent.h
+++ b/Engine/Source/Editor/GameplayAbilitiesEditor/Classes/K2Node_GameplayCueEvent.h
@@ -7,9 +7,7 @@
 UCLASS()
 class UK2Node_GameplayCueEvent : public UK2Node_Event
 {
-	GENERATED_BODY()
-public:
-	UK2Node_GameplayCueEvent(const FObjectInitializer& ObjectInitializer = FObjectInitializer::Get());
+	GENERATED_UCLASS_BODY()
 
 	// UObject interface
 	virtual void Serialize(FArchive& Ar) override;
@@ -25,4 +23,4 @@ public:
 	// Begin UK2Node interface
 	virtual void GetMenuActions(FBlueprintActionDatabaseRegistrar& ActionRegistrar) const override;
 	// End UK2Node interface
-};
+};
\ No newline at end of file
diff --git a/Engine/Source/Editor/GameplayAbilitiesEditor/Classes/K2Node_GameplayEffectVariable.h b/Engine/Source/Editor/GameplayAbilitiesEditor/Classes/K2Node_GameplayEffectVariable.h
--- a/Engine/Source/Editor/GameplayAbilitiesEditor/Classes/K2Node_GameplayEffectVariable.h
+++ b/Engine/Source/Editor/GameplayAbilitiesEditor/Classes/K2Node_GameplayEffectVariable.h
@@ -11,9 +11,7 @@
 UCLASS()
 class GAMEPLAYABILITIESEDITOR_API UK2Node_GameplayEffectVariable : public UK2Node_VariableGet // only need a custom node for the getter
 {
-	GENERATED_BODY()
-public:
-	UK2Node_GameplayEffectVariable(const FObjectInitializer& ObjectInitializer = FObjectInitializer::Get());
+	GENERATED_UCLASS_BODY()
 
 public:
 	UPROPERTY()
diff --git a/Engine/Source/Editor/GameplayAbilitiesEditor/Classes/K2Node_LatentAbilityCall.h b/Engine/Source/Editor/GameplayAbilitiesEditor/Classes/K2Node_LatentAbilityCall.h
--- a/Engine/Source/Editor/GameplayAbilitiesEditor/Classes/K2Node_LatentAbilityCall.h
+++ b/Engine/Source/Editor/GameplayAbilitiesEditor/Classes/K2Node_LatentAbilityCall.h
@@ -11,9 +11,7 @@
 UCLASS()
 class UK2Node_LatentAbilityCall : public UK2Node_BaseAsyncTask
 {
-	GENERATED_BODY()
-public:
-	UK2Node_LatentAbilityCall(const FObjectInitializer& ObjectInitializer = FObjectInitializer::Get());
+	GENERATED_UCLASS_BODY()
 
 	// UEdGraphNode interface
 	virtual bool CanCreateUnderSpecifiedSchema(const UEdGraphSchema* DesiredSchema) const;
diff --git a/Engine/Source/Editor/GameplayTagsEditor/Classes/GameplayTagsK2Node_LiteralGameplayTag.h b/Engine/Source/Editor/GameplayTagsEditor/Classes/GameplayTagsK2Node_LiteralGameplayTag.h
--- a/Engine/Source/Editor/GameplayTagsEditor/Classes/GameplayTagsK2Node_LiteralGameplayTag.h
+++ b/Engine/Source/Editor/GameplayTagsEditor/Classes/GameplayTagsK2Node_LiteralGameplayTag.h
@@ -8,9 +8,7 @@
 UCLASS()
 class UGameplayTagsK2Node_LiteralGameplayTag : public UK2Node
 {
-	GENERATED_BODY()
-public:
-	UGameplayTagsK2Node_LiteralGameplayTag(const FObjectInitializer& ObjectInitializer = FObjectInitializer::Get());
+	GENERATED_UCLASS_BODY()
 
 #if WITH_EDITOR
 	// Begin UEdGraphNode interface
diff --git a/Engine/Source/Editor/GameplayTagsEditor/Classes/GameplayTagsK2Node_MultiCompareGameplayTagAssetInterfaceSingleTags.h b/Engine/Source/Editor/GameplayTagsEditor/Classes/GameplayTagsK2Node_MultiCompareGameplayTagAssetInterfaceSingleTags.h
--- a/Engine/Source/Editor/GameplayTagsEditor/Classes/GameplayTagsK2Node_MultiCompareGameplayTagAssetInterfaceSingleTags.h
+++ b/Engine/Source/Editor/GameplayTagsEditor/Classes/GameplayTagsK2Node_MultiCompareGameplayTagAssetInterfaceSingleTags.h
@@ -8,9 +8,7 @@
 UCLASS()
 class UGameplayTagsK2Node_MultiCompareGameplayTagAssetInterfaceSingleTags : public UGameplayTagsK2Node_MultiCompareBase
 {
-	GENERATED_BODY()
-public:
-	UGameplayTagsK2Node_MultiCompareGameplayTagAssetInterfaceSingleTags(const FObjectInitializer& ObjectInitializer = FObjectInitializer::Get());
+	GENERATED_UCLASS_BODY()
 
 	// UEdGraphNode interface
 	virtual void AllocateDefaultPins() override;
diff --git a/Engine/Source/Editor/GameplayTagsEditor/Classes/GameplayTagsK2Node_MultiCompareGameplayTagContainer.h b/Engine/Source/Editor/GameplayTagsEditor/Classes/GameplayTagsK2Node_MultiCompareGameplayTagContainer.h
--- a/Engine/Source/Editor/GameplayTagsEditor/Classes/GameplayTagsK2Node_MultiCompareGameplayTagContainer.h
+++ b/Engine/Source/Editor/GameplayTagsEditor/Classes/GameplayTagsK2Node_MultiCompareGameplayTagContainer.h
@@ -8,9 +8,7 @@
 UCLASS()
 class UGameplayTagsK2Node_MultiCompareGameplayTagContainer : public UGameplayTagsK2Node_MultiCompareBase
 {
-	GENERATED_BODY()
-public:
-	UGameplayTagsK2Node_MultiCompareGameplayTagContainer(const FObjectInitializer& ObjectInitializer = FObjectInitializer::Get());
+	GENERATED_UCLASS_BODY()
 
 	// UEdGraphNode interface
 	virtual void AllocateDefaultPins() override;
diff --git a/Engine/Source/Editor/GameplayTagsEditor/Classes/GameplayTagsK2Node_MultiCompareGameplayTagContainerSingleTags.h b/Engine/Source/Editor/GameplayTagsEditor/Classes/GameplayTagsK2Node_MultiCompareGameplayTagContainerSingleTags.h
--- a/Engine/Source/Editor/GameplayTagsEditor/Classes/GameplayTagsK2Node_MultiCompareGameplayTagContainerSingleTags.h
+++ b/Engine/Source/Editor/GameplayTagsEditor/Classes/GameplayTagsK2Node_MultiCompareGameplayTagContainerSingleTags.h
@@ -8,9 +8,7 @@
 UCLASS()
 class UGameplayTagsK2Node_MultiCompareGameplayTagContainerSingleTags : public UGameplayTagsK2Node_MultiCompareBase
 {
-	GENERATED_BODY()
-public:
-	UGameplayTagsK2Node_MultiCompareGameplayTagContainerSingleTags(const FObjectInitializer& ObjectInitializer = FObjectInitializer::Get());
+	GENERATED_UCLASS_BODY()
 
 	// UEdGraphNode interface
 	virtual void AllocateDefaultPins() override;
diff --git a/Engine/Source/Editor/GameplayTagsEditor/Classes/GameplayTagsK2Node_SwitchGameplayTag.h b/Engine/Source/Editor/GameplayTagsEditor/Classes/GameplayTagsK2Node_SwitchGameplayTag.h
--- a/Engine/Source/Editor/GameplayTagsEditor/Classes/GameplayTagsK2Node_SwitchGameplayTag.h
+++ b/Engine/Source/Editor/GameplayTagsEditor/Classes/GameplayTagsK2Node_SwitchGameplayTag.h
@@ -11,9 +11,7 @@ struct FGameplayTag;
 UCLASS(MinimalAPI)
 class UGameplayTagsK2Node_SwitchGameplayTag : public UK2Node_Switch
 {
-	GENERATED_BODY()
-public:
-	GAMEPLAYTAGSEDITOR_API UGameplayTagsK2Node_SwitchGameplayTag(const FObjectInitializer& ObjectInitializer = FObjectInitializer::Get());
+	GENERATED_UCLASS_BODY()
 
 	UPROPERTY(EditAnywhere, Category = PinOptions)
 	TArray<FGameplayTag> PinTags;
diff --git a/Engine/Source/Editor/GameplayTagsEditor/Classes/GameplayTagsK2Node_SwitchGameplayTagContainer.h b/Engine/Source/Editor/GameplayTagsEditor/Classes/GameplayTagsK2Node_SwitchGameplayTagContainer.h
--- a/Engine/Source/Editor/GameplayTagsEditor/Classes/GameplayTagsK2Node_SwitchGameplayTagContainer.h
+++ b/Engine/Source/Editor/GameplayTagsEditor/Classes/GameplayTagsK2Node_SwitchGameplayTagContainer.h
@@ -11,9 +11,7 @@ struct FGameplayTagContainer;
 UCLASS(MinimalAPI)
 class UGameplayTagsK2Node_SwitchGameplayTagContainer : public UK2Node_Switch
 {
-	GENERATED_BODY()
-public:
-	GAMEPLAYTAGSEDITOR_API UGameplayTagsK2Node_SwitchGameplayTagContainer(const FObjectInitializer& ObjectInitializer = FObjectInitializer::Get());
+	GENERATED_UCLASS_BODY()
 
 	UPROPERTY(EditAnywhere, Category = PinOptions)
 	TArray<FGameplayTagContainer> PinContainers;
diff --git a/Engine/Source/Editor/GeometryMode/Classes/GeomModifier.h b/Engine/Source/Editor/GeometryMode/Classes/GeomModifier.h
--- a/Engine/Source/Editor/GeometryMode/Classes/GeomModifier.h
+++ b/Engine/Source/Editor/GeometryMode/Classes/GeomModifier.h
@@ -11,9 +11,7 @@
 UCLASS(abstract, hidecategories=(Object, GeomModifier))
 class UGeomModifier : public UObject
 {
-	GENERATED_BODY()
-public:
-	UGeomModifier(const FObjectInitializer& ObjectInitializer = FObjectInitializer::Get());
+	GENERATED_UCLASS_BODY()
 
 	/** A human readable name for this modifier (appears on buttons, menus, etc) */
 	UPROPERTY(EditAnywhere, Category=GeomModifier)
diff --git a/Engine/Source/Editor/GeometryMode/Classes/GeomModifier_Create.h b/Engine/Source/Editor/GeometryMode/Classes/GeomModifier_Create.h
--- a/Engine/Source/Editor/GeometryMode/Classes/GeomModifier_Create.h
+++ b/Engine/Source/Editor/GeometryMode/Classes/GeomModifier_Create.h
@@ -11,9 +11,7 @@
 UCLASS()
 class UGeomModifier_Create : public UGeomModifier_Edit
 {
-	GENERATED_BODY()
-public:
-	UGeomModifier_Create(const FObjectInitializer& ObjectInitializer = FObjectInitializer::Get());
+	GENERATED_UCLASS_BODY()
 
 
 	// Begin UGeomModifier Interface
diff --git a/Engine/Source/Editor/GeometryMode/Classes/GeomModifier_Delete.h b/Engine/Source/Editor/GeometryMode/Classes/GeomModifier_Delete.h
--- a/Engine/Source/Editor/GeometryMode/Classes/GeomModifier_Delete.h
+++ b/Engine/Source/Editor/GeometryMode/Classes/GeomModifier_Delete.h
@@ -11,9 +11,7 @@
 UCLASS()
 class UGeomModifier_Delete : public UGeomModifier_Edit
 {
-	GENERATED_BODY()
-public:
-	UGeomModifier_Delete(const FObjectInitializer& ObjectInitializer = FObjectInitializer::Get());
+	GENERATED_UCLASS_BODY()
 
 
 	// Begin UGeomModifier Interface
diff --git a/Engine/Source/Editor/GeometryMode/Classes/GeomModifier_Edit.h b/Engine/Source/Editor/GeometryMode/Classes/GeomModifier_Edit.h
--- a/Engine/Source/Editor/GeometryMode/Classes/GeomModifier_Edit.h
+++ b/Engine/Source/Editor/GeometryMode/Classes/GeomModifier_Edit.h
@@ -16,9 +16,7 @@
 UCLASS(autoexpandcategories=Settings)
 class UGeomModifier_Edit : public UGeomModifier
 {
-	GENERATED_BODY()
-public:
-	UGeomModifier_Edit(const FObjectInitializer& ObjectInitializer = FObjectInitializer::Get());
+	GENERATED_UCLASS_BODY()
 
 
 
diff --git a/Engine/Source/Editor/GeometryMode/Classes/GeomModifier_Extrude.h b/Engine/Source/Editor/GeometryMode/Classes/GeomModifier_Extrude.h
--- a/Engine/Source/Editor/GeometryMode/Classes/GeomModifier_Extrude.h
+++ b/Engine/Source/Editor/GeometryMode/Classes/GeomModifier_Extrude.h
@@ -11,9 +11,7 @@
 UCLASS()
 class UGeomModifier_Extrude : public UGeomModifier_Edit
 {
-	GENERATED_BODY()
-public:
-	UGeomModifier_Extrude(const FObjectInitializer& ObjectInitializer = FObjectInitializer::Get());
+	GENERATED_UCLASS_BODY()
 
 	UPROPERTY(EditAnywhere, Category=Settings)
 	int32 Length;
diff --git a/Engine/Source/Editor/GeometryMode/Classes/GeomModifier_Flip.h b/Engine/Source/Editor/GeometryMode/Classes/GeomModifier_Flip.h
--- a/Engine/Source/Editor/GeometryMode/Classes/GeomModifier_Flip.h
+++ b/Engine/Source/Editor/GeometryMode/Classes/GeomModifier_Flip.h
@@ -11,9 +11,7 @@
 UCLASS()
 class UGeomModifier_Flip : public UGeomModifier_Edit
 {
-	GENERATED_BODY()
-public:
-	UGeomModifier_Flip(const FObjectInitializer& ObjectInitializer = FObjectInitializer::Get());
+	GENERATED_UCLASS_BODY()
 
 
 	// Begin UGeomModifier Interface
diff --git a/Engine/Source/Editor/GeometryMode/Classes/GeomModifier_Lathe.h b/Engine/Source/Editor/GeometryMode/Classes/GeomModifier_Lathe.h
--- a/Engine/Source/Editor/GeometryMode/Classes/GeomModifier_Lathe.h
+++ b/Engine/Source/Editor/GeometryMode/Classes/GeomModifier_Lathe.h
@@ -11,9 +11,7 @@
 UCLASS()
 class UGeomModifier_Lathe : public UGeomModifier_Edit
 {
-	GENERATED_BODY()
-public:
-	UGeomModifier_Lathe(const FObjectInitializer& ObjectInitializer = FObjectInitializer::Get());
+	GENERATED_UCLASS_BODY()
 
 	UPROPERTY(EditAnywhere, Category=Settings)
 	int32 TotalSegments;
diff --git a/Engine/Source/Editor/GeometryMode/Classes/GeomModifier_Optimize.h b/Engine/Source/Editor/GeometryMode/Classes/GeomModifier_Optimize.h
--- a/Engine/Source/Editor/GeometryMode/Classes/GeomModifier_Optimize.h
+++ b/Engine/Source/Editor/GeometryMode/Classes/GeomModifier_Optimize.h
@@ -11,9 +11,7 @@
 UCLASS()
 class UGeomModifier_Optimize : public UGeomModifier_Edit
 {
-	GENERATED_BODY()
-public:
-	UGeomModifier_Optimize(const FObjectInitializer& ObjectInitializer = FObjectInitializer::Get());
+	GENERATED_UCLASS_BODY()
 
 
 	// Begin UGeomModifier Interface
diff --git a/Engine/Source/Editor/GeometryMode/Classes/GeomModifier_Pen.h b/Engine/Source/Editor/GeometryMode/Classes/GeomModifier_Pen.h
--- a/Engine/Source/Editor/GeometryMode/Classes/GeomModifier_Pen.h
+++ b/Engine/Source/Editor/GeometryMode/Classes/GeomModifier_Pen.h
@@ -11,9 +11,7 @@
 UCLASS()
 class UGeomModifier_Pen : public UGeomModifier_Edit
 {
-	GENERATED_BODY()
-public:
-	UGeomModifier_Pen(const FObjectInitializer& ObjectInitializer = FObjectInitializer::Get());
+	GENERATED_UCLASS_BODY()
 
 	/** If true, the shape will be automatically extruded into a brush upon completion. */
 	UPROPERTY(EditAnywhere, Category=Settings)
diff --git a/Engine/Source/Editor/GeometryMode/Classes/GeomModifier_Split.h b/Engine/Source/Editor/GeometryMode/Classes/GeomModifier_Split.h
--- a/Engine/Source/Editor/GeometryMode/Classes/GeomModifier_Split.h
+++ b/Engine/Source/Editor/GeometryMode/Classes/GeomModifier_Split.h
@@ -11,9 +11,7 @@
 UCLASS()
 class UGeomModifier_Split : public UGeomModifier_Edit
 {
-	GENERATED_BODY()
-public:
-	UGeomModifier_Split(const FObjectInitializer& ObjectInitializer = FObjectInitializer::Get());
+	GENERATED_UCLASS_BODY()
 
 
 	// Begin UGeomModifier Interface
diff --git a/Engine/Source/Editor/GeometryMode/Classes/GeomModifier_Triangulate.h b/Engine/Source/Editor/GeometryMode/Classes/GeomModifier_Triangulate.h
--- a/Engine/Source/Editor/GeometryMode/Classes/GeomModifier_Triangulate.h
+++ b/Engine/Source/Editor/GeometryMode/Classes/GeomModifier_Triangulate.h
@@ -11,9 +11,7 @@
 UCLASS()
 class UGeomModifier_Triangulate : public UGeomModifier_Edit
 {
-	GENERATED_BODY()
-public:
-	UGeomModifier_Triangulate(const FObjectInitializer& ObjectInitializer = FObjectInitializer::Get());
+	GENERATED_UCLASS_BODY()
 
 
 	// Begin UGeomModifier Interface
diff --git a/Engine/Source/Editor/GeometryMode/Classes/GeomModifier_Turn.h b/Engine/Source/Editor/GeometryMode/Classes/GeomModifier_Turn.h
--- a/Engine/Source/Editor/GeometryMode/Classes/GeomModifier_Turn.h
+++ b/Engine/Source/Editor/GeometryMode/Classes/GeomModifier_Turn.h
@@ -11,9 +11,7 @@
 UCLASS()
 class UGeomModifier_Turn : public UGeomModifier_Edit
 {
-	GENERATED_BODY()
-public:
-	UGeomModifier_Turn(const FObjectInitializer& ObjectInitializer = FObjectInitializer::Get());
+	GENERATED_UCLASS_BODY()
 
 
 	// Begin UGeomModifier Interface
diff --git a/Engine/Source/Editor/GeometryMode/Classes/GeomModifier_Weld.h b/Engine/Source/Editor/GeometryMode/Classes/GeomModifier_Weld.h
--- a/Engine/Source/Editor/GeometryMode/Classes/GeomModifier_Weld.h
+++ b/Engine/Source/Editor/GeometryMode/Classes/GeomModifier_Weld.h
@@ -10,9 +10,7 @@
 UCLASS()
 class UGeomModifier_Weld : public UGeomModifier_Edit
 {
-	GENERATED_BODY()
-public:
-	UGeomModifier_Weld(const FObjectInitializer& ObjectInitializer = FObjectInitializer::Get());
+	GENERATED_UCLASS_BODY()
 
 
 	// Begin UGeomModifier Interface
diff --git a/Engine/Source/Editor/GraphEditor/Public/GraphEditorSettings.h b/Engine/Source/Editor/GraphEditor/Public/GraphEditorSettings.h
--- a/Engine/Source/Editor/GraphEditor/Public/GraphEditorSettings.h
+++ b/Engine/Source/Editor/GraphEditor/Public/GraphEditorSettings.h
@@ -13,9 +13,7 @@ UCLASS(config=EditorUserSettings)
 class GRAPHEDITOR_API UGraphEditorSettings
 	:	public UObject
 {
-	GENERATED_BODY()
-public:
-	UGraphEditorSettings(const FObjectInitializer& ObjectInitializer = FObjectInitializer::Get());
+	GENERATED_UCLASS_BODY()
 
 public:
 
diff --git a/Engine/Source/Editor/HardwareTargeting/Public/HardwareTargetingSettings.h b/Engine/Source/Editor/HardwareTargeting/Public/HardwareTargetingSettings.h
--- a/Engine/Source/Editor/HardwareTargeting/Public/HardwareTargetingSettings.h
+++ b/Engine/Source/Editor/HardwareTargeting/Public/HardwareTargetingSettings.h
@@ -42,9 +42,7 @@ namespace EGraphicsPreset
 UCLASS(config=Engine, defaultconfig)
 class HARDWARETARGETING_API UHardwareTargetingSettings : public UObject
 {
-	GENERATED_BODY()
-public:
-	UHardwareTargetingSettings(const FObjectInitializer& ObjectInitializer = FObjectInitializer::Get());
+	GENERATED_UCLASS_BODY()
 
 	/** Enum specifying which class of hardware this game is targeting */
 	UPROPERTY(config, EditAnywhere, category=None)
@@ -79,4 +77,4 @@ private:
 	/** Holds an event delegate that is executed when a setting has changed. */
 	FSettingChangedEvent SettingChangedEvent;
 #endif
-};
+};
\ No newline at end of file
diff --git a/Engine/Source/Editor/IntroTutorials/Private/EditorTutorialFactory.h b/Engine/Source/Editor/IntroTutorials/Private/EditorTutorialFactory.h
--- a/Engine/Source/Editor/IntroTutorials/Private/EditorTutorialFactory.h
+++ b/Engine/Source/Editor/IntroTutorials/Private/EditorTutorialFactory.h
@@ -7,9 +7,7 @@
 UCLASS(hidecategories=Object, collapsecategories)
 class UEditorTutorialFactory : public UFactory
 {
-	GENERATED_BODY()
-public:
-	UEditorTutorialFactory(const FObjectInitializer& ObjectInitializer = FObjectInitializer::Get());
+	GENERATED_UCLASS_BODY()
 
 	// UFactory Interface
 	virtual UObject* FactoryCreateNew(UClass* Class, UObject* InParent, FName Name, EObjectFlags Flags, UObject* Context, FFeedbackContext* Warn) override;
diff --git a/Engine/Source/Editor/IntroTutorials/Private/EditorTutorialImportFactory.h b/Engine/Source/Editor/IntroTutorials/Private/EditorTutorialImportFactory.h
--- a/Engine/Source/Editor/IntroTutorials/Private/EditorTutorialImportFactory.h
+++ b/Engine/Source/Editor/IntroTutorials/Private/EditorTutorialImportFactory.h
@@ -7,9 +7,7 @@
 UCLASS(hidecategories=Object, collapsecategories)
 class UEditorTutorialImportFactory : public UFactory, public FReimportHandler
 {
-	GENERATED_BODY()
-public:
-	UEditorTutorialImportFactory(const FObjectInitializer& ObjectInitializer = FObjectInitializer::Get());
+	GENERATED_UCLASS_BODY()
 
 	// UFactory Interface
 	virtual bool FactoryCanImport(const FString& Filename) override;
diff --git a/Engine/Source/Editor/IntroTutorials/Private/EditorTutorialSettings.h b/Engine/Source/Editor/IntroTutorials/Private/EditorTutorialSettings.h
--- a/Engine/Source/Editor/IntroTutorials/Private/EditorTutorialSettings.h
+++ b/Engine/Source/Editor/IntroTutorials/Private/EditorTutorialSettings.h
@@ -9,7 +9,7 @@
 USTRUCT()
 struct FTutorialContext
 {
-	GENERATED_BODY()
+	GENERATED_USTRUCT_BODY()
 
 	/** The context that this tutorial is used in */
 	UPROPERTY(EditAnywhere, Category = "Tutorials")
@@ -32,9 +32,7 @@ struct FTutorialContext
 UCLASS(config=EditorGameAgnostic)
 class UEditorTutorialSettings : public UObject
 {
-	GENERATED_BODY()
-public:
-	UEditorTutorialSettings(const FObjectInitializer& ObjectInitializer = FObjectInitializer::Get());
+	GENERATED_UCLASS_BODY()
 
 	/** Categories for tutorials */
 	UPROPERTY(Config, EditAnywhere, Category="Tutorials")
@@ -50,4 +48,4 @@ public:
 
 	/** Get the tutorial info for the specified context */
 	void FindTutorialInfoForContext(FName InContext, UEditorTutorial*& OutAttractTutorial, UEditorTutorial*& OutLaunchTutorial, FString& OutBrowserFilter) const;
-};
+};
\ No newline at end of file
diff --git a/Engine/Source/Editor/IntroTutorials/Private/TutorialSettings.h b/Engine/Source/Editor/IntroTutorials/Private/TutorialSettings.h
--- a/Engine/Source/Editor/IntroTutorials/Private/TutorialSettings.h
+++ b/Engine/Source/Editor/IntroTutorials/Private/TutorialSettings.h
@@ -9,9 +9,7 @@
 UCLASS(config=Editor)
 class UTutorialSettings : public UObject
 {
-	GENERATED_BODY()
-public:
-	UTutorialSettings(const FObjectInitializer& ObjectInitializer = FObjectInitializer::Get());
+	GENERATED_UCLASS_BODY()
 
 	/** Categories for tutorials */
 	UPROPERTY(Config, EditAnywhere, Category="Tutorials")
@@ -20,4 +18,4 @@ public:
 	/** Tutorial to start on project startup */
 	UPROPERTY(Config, EditAnywhere, Category="Tutorials", meta=(MetaClass="EditorTutorial"))
 	FStringClassReference StartupTutorial;
-};
+};
\ No newline at end of file
diff --git a/Engine/Source/Editor/IntroTutorials/Private/TutorialStateSettings.h b/Engine/Source/Editor/IntroTutorials/Private/TutorialStateSettings.h
--- a/Engine/Source/Editor/IntroTutorials/Private/TutorialStateSettings.h
+++ b/Engine/Source/Editor/IntroTutorials/Private/TutorialStateSettings.h
@@ -9,7 +9,7 @@
 USTRUCT()
 struct FTutorialProgress
 {
-	GENERATED_BODY()
+	GENERATED_USTRUCT_BODY()
 
 	FTutorialProgress()
 	{
@@ -33,9 +33,7 @@ struct FTutorialProgress
 UCLASS(config=EditorGameAgnostic)
 class UTutorialStateSettings : public UObject
 {
-	GENERATED_BODY()
-public:
-	UTutorialStateSettings(const FObjectInitializer& ObjectInitializer = FObjectInitializer::Get());
+	GENERATED_UCLASS_BODY()
 
 	UPROPERTY(Config)
 	TArray<FTutorialProgress> TutorialsProgress;
@@ -70,4 +68,4 @@ public:
 private:
 	/** Recorded progress */
 	TMap<UEditorTutorial*, FTutorialProgress> ProgressMap;
-};
+};
\ No newline at end of file
diff --git a/Engine/Source/Editor/IntroTutorials/Public/EditorTutorial.h b/Engine/Source/Editor/IntroTutorials/Public/EditorTutorial.h
--- a/Engine/Source/Editor/IntroTutorials/Public/EditorTutorial.h
+++ b/Engine/Source/Editor/IntroTutorials/Public/EditorTutorial.h
@@ -45,7 +45,7 @@ namespace ETutorialAnchorIdentifier
 USTRUCT()
 struct INTROTUTORIALS_API FTutorialCategory
 {
-	GENERATED_BODY()
+	GENERATED_USTRUCT_BODY()
 
 	/** Period-separated category name, e.g. "Editor Quickstart.Level Editor" */
 	UPROPERTY(EditAnywhere, Category="Content")
@@ -76,7 +76,7 @@ struct INTROTUTORIALS_API FTutorialCategory
 USTRUCT()
 struct INTROTUTORIALS_API FTutorialContent
 {
-	GENERATED_BODY()
+	GENERATED_USTRUCT_BODY()
 
 	FTutorialContent()
 	{
@@ -104,7 +104,7 @@ struct INTROTUTORIALS_API FTutorialContent
 USTRUCT()
 struct INTROTUTORIALS_API FTutorialContentAnchor
 {
-	GENERATED_BODY()
+	GENERATED_USTRUCT_BODY()
 	
 	FTutorialContentAnchor()
 	{
@@ -149,7 +149,7 @@ struct INTROTUTORIALS_API FTutorialContentAnchor
 USTRUCT()
 struct INTROTUTORIALS_API FTutorialWidgetContent
 {
-	GENERATED_BODY()
+	GENERATED_USTRUCT_BODY()
 
 	FTutorialWidgetContent()
 	{
@@ -190,7 +190,7 @@ struct INTROTUTORIALS_API FTutorialWidgetContent
 USTRUCT()
 struct INTROTUTORIALS_API FTutorialStage
 {
-	GENERATED_BODY()
+	GENERATED_USTRUCT_BODY()
 
 	/** Identifier for this stage */
 	UPROPERTY(EditAnywhere, Category="Stage")
@@ -217,9 +217,7 @@ struct INTROTUTORIALS_API FTutorialStage
 UCLASS(Blueprintable, hideCategories=(Object), EditInlineNew)
 class INTROTUTORIALS_API UEditorTutorial : public UObject
 {
-	GENERATED_BODY()
-public:
-	UEditorTutorial(const FObjectInitializer& ObjectInitializer = FObjectInitializer::Get());
+	GENERATED_UCLASS_BODY()
 
 	/** Title of this tutorial, used when presented to the user */
 	UPROPERTY(EditAnywhere, Category="Tutorial", AssetRegistrySearchable)
@@ -347,4 +345,4 @@ protected:
 	 */
 	UFUNCTION(BlueprintCallable, Category="Tutorial")
 	static bool GetEngineFolderVisibilty();
-};
+};
\ No newline at end of file
diff --git a/Engine/Source/Editor/Kismet/Public/BlueprintActionMenuItem.h b/Engine/Source/Editor/Kismet/Public/BlueprintActionMenuItem.h
--- a/Engine/Source/Editor/Kismet/Public/BlueprintActionMenuItem.h
+++ b/Engine/Source/Editor/Kismet/Public/BlueprintActionMenuItem.h
@@ -22,7 +22,7 @@ struct FBlueprintActionContext;
 USTRUCT()
 struct FBlueprintActionMenuItem : public FEdGraphSchemaAction
 {
-	GENERATED_BODY()
+	GENERATED_USTRUCT_BODY();
 	static FName StaticGetTypeId() { static FName const TypeId("FBlueprintActionMenuItem"); return TypeId; }
 
 public:	
diff --git a/Engine/Source/Editor/Kismet/Public/Nodes/K2Node_AsyncAction.h b/Engine/Source/Editor/Kismet/Public/Nodes/K2Node_AsyncAction.h
--- a/Engine/Source/Editor/Kismet/Public/Nodes/K2Node_AsyncAction.h
+++ b/Engine/Source/Editor/Kismet/Public/Nodes/K2Node_AsyncAction.h
@@ -7,9 +7,7 @@
 UCLASS()
 class KISMET_API UK2Node_AsyncAction : public UK2Node_BaseAsyncTask
 {
-	GENERATED_BODY()
-public:
-	UK2Node_AsyncAction(const FObjectInitializer& ObjectInitializer = FObjectInitializer::Get());
+	GENERATED_UCLASS_BODY()
 	
 	// UEdGraphNode interface
 	virtual void GetMenuEntries(FGraphContextMenuBuilder& ContextMenuBuilder) const override;
diff --git a/Engine/Source/Editor/LandscapeEditor/Classes/ActorFactoryLandscape.h b/Engine/Source/Editor/LandscapeEditor/Classes/ActorFactoryLandscape.h
--- a/Engine/Source/Editor/LandscapeEditor/Classes/ActorFactoryLandscape.h
+++ b/Engine/Source/Editor/LandscapeEditor/Classes/ActorFactoryLandscape.h
@@ -7,9 +7,7 @@
 UCLASS(MinimalAPI, config=Editor)
 class UActorFactoryLandscape : public UActorFactory
 {
-	GENERATED_BODY()
-public:
-	LANDSCAPEEDITOR_API UActorFactoryLandscape(const FObjectInitializer& ObjectInitializer = FObjectInitializer::Get());
+	GENERATED_UCLASS_BODY()
 
 	// Begin UActorFactory Interface
 	//virtual bool CanCreateActorFrom( const FAssetData& AssetData, FText& OutErrorMsg ) override;
diff --git a/Engine/Source/Editor/LandscapeEditor/Classes/LandscapeEditorObject.h b/Engine/Source/Editor/LandscapeEditor/Classes/LandscapeEditorObject.h
--- a/Engine/Source/Editor/LandscapeEditor/Classes/LandscapeEditorObject.h
+++ b/Engine/Source/Editor/LandscapeEditor/Classes/LandscapeEditorObject.h
@@ -128,7 +128,7 @@ namespace EColorChannel
 USTRUCT()
 struct FGizmoImportLayer
 {
-	GENERATED_BODY()
+	GENERATED_USTRUCT_BODY()
 
 	UPROPERTY(Category="Import", EditAnywhere)
 	FString LayerFilename;
@@ -175,7 +175,7 @@ enum class ELandscapeImportLayerError : uint8
 USTRUCT()
 struct FLandscapeImportLayer : public FLandscapeImportLayerInfo
 {
-	GENERATED_BODY()
+	GENERATED_USTRUCT_BODY()
 
 	UPROPERTY(Category="Import", VisibleAnywhere)
 	ELandscapeImportLayerError ImportError;
@@ -190,7 +190,7 @@ struct FLandscapeImportLayer : public FLandscapeImportLayerInfo
 USTRUCT()
 struct FLandscapePatternBrushWorldSpaceSettings
 {
-	GENERATED_BODY()
+	GENERATED_USTRUCT_BODY()
 
 	UPROPERTY(Category="World-Space", EditAnywhere)
 	FVector2D Origin;
@@ -212,9 +212,7 @@ struct FLandscapePatternBrushWorldSpaceSettings
 UCLASS()
 class ULandscapeEditorObject : public UObject
 {
-	GENERATED_BODY()
-public:
-	ULandscapeEditorObject(const FObjectInitializer& ObjectInitializer = FObjectInitializer::Get());
+	GENERATED_UCLASS_BODY()
 
 	FEdModeLandscape* ParentMode;
 
diff --git a/Engine/Source/Editor/LandscapeEditor/Classes/LandscapePlaceholder.h b/Engine/Source/Editor/LandscapeEditor/Classes/LandscapePlaceholder.h
--- a/Engine/Source/Editor/LandscapeEditor/Classes/LandscapePlaceholder.h
+++ b/Engine/Source/Editor/LandscapeEditor/Classes/LandscapePlaceholder.h
@@ -7,9 +7,7 @@
 UCLASS(NotPlaceable)
 class ALandscapePlaceholder : public AActor
 {
-	GENERATED_BODY()
-public:
-	ALandscapePlaceholder(const FObjectInitializer& ObjectInitializer = FObjectInitializer::Get());
+	GENERATED_UCLASS_BODY()
 
 	//virtual void PostEditMove(bool bFinished) override;
 	virtual bool TeleportTo(const FVector& DestLocation, const FRotator& DestRotation, bool bIsATest = false, bool bNoCheck = false) override;
diff --git a/Engine/Source/Editor/Levels/Public/LevelBrowserSettings.h b/Engine/Source/Editor/Levels/Public/LevelBrowserSettings.h
--- a/Engine/Source/Editor/Levels/Public/LevelBrowserSettings.h
+++ b/Engine/Source/Editor/Levels/Public/LevelBrowserSettings.h
@@ -16,9 +16,7 @@ UCLASS(config=EditorUserSettings)
 class LEVELS_API ULevelBrowserSettings
 	:	public UObject
 {
-	GENERATED_BODY()
-public:
-	ULevelBrowserSettings(const FObjectInitializer& ObjectInitializer = FObjectInitializer::Get());
+	GENERATED_UCLASS_BODY()
 
 public:
 
diff --git a/Engine/Source/Editor/Matinee/Classes/InterpTrackHelper.h b/Engine/Source/Editor/Matinee/Classes/InterpTrackHelper.h
--- a/Engine/Source/Editor/Matinee/Classes/InterpTrackHelper.h
+++ b/Engine/Source/Editor/Matinee/Classes/InterpTrackHelper.h
@@ -8,9 +8,7 @@
 UCLASS()
 class MATINEE_API UInterpTrackHelper : public UObject
 {
-	GENERATED_BODY()
-public:
-	UInterpTrackHelper(const FObjectInitializer& ObjectInitializer = FObjectInitializer::Get());
+	GENERATED_UCLASS_BODY()
 
 	/**
 	 * @param  Track	The track to get the actor for.
diff --git a/Engine/Source/Editor/Matinee/Classes/MatineeOptions.h b/Engine/Source/Editor/Matinee/Classes/MatineeOptions.h
--- a/Engine/Source/Editor/Matinee/Classes/MatineeOptions.h
+++ b/Engine/Source/Editor/Matinee/Classes/MatineeOptions.h
@@ -11,9 +11,7 @@
 UCLASS(hidecategories=Object, config=Editor)
 class UMatineeOptions : public UObject
 {
-	GENERATED_BODY()
-public:
-	UMatineeOptions(const FObjectInitializer& ObjectInitializer = FObjectInitializer::Get());
+	GENERATED_UCLASS_BODY()
 
 	UPROPERTY()
 	TArray<struct FInterpEdSelKey> SelectedKeys;
diff --git a/Engine/Source/Editor/Matinee/Classes/MatineeTrackAnimControlHelper.h b/Engine/Source/Editor/Matinee/Classes/MatineeTrackAnimControlHelper.h
--- a/Engine/Source/Editor/Matinee/Classes/MatineeTrackAnimControlHelper.h
+++ b/Engine/Source/Editor/Matinee/Classes/MatineeTrackAnimControlHelper.h
@@ -8,9 +8,7 @@
 UCLASS()
 class UMatineeTrackAnimControlHelper : public UInterpTrackHelper
 {
-	GENERATED_BODY()
-public:
-	UMatineeTrackAnimControlHelper(const FObjectInitializer& ObjectInitializer = FObjectInitializer::Get());
+	GENERATED_UCLASS_BODY()
 
 	void OnAddKeyTextEntry(const class FAssetData& AssetData, IMatineeBase* Matinee, UInterpTrack* Track);
 	void OnCreateTrackTextEntry(const FString& ChosenText, TSharedRef<SWindow> Window, FString* OutputString);
diff --git a/Engine/Source/Editor/Matinee/Classes/MatineeTrackBoolPropHelper.h b/Engine/Source/Editor/Matinee/Classes/MatineeTrackBoolPropHelper.h
--- a/Engine/Source/Editor/Matinee/Classes/MatineeTrackBoolPropHelper.h
+++ b/Engine/Source/Editor/Matinee/Classes/MatineeTrackBoolPropHelper.h
@@ -7,9 +7,7 @@
 UCLASS()
 class UMatineeTrackBoolPropHelper : public UInterpTrackHelper
 {
-	GENERATED_BODY()
-public:
-	UMatineeTrackBoolPropHelper(const FObjectInitializer& ObjectInitializer = FObjectInitializer::Get());
+	GENERATED_UCLASS_BODY()
 
 	void OnCreateTrackTextEntry(const FString& ChosenText, TWeakPtr<SWindow> Window, FString* OutputString);
 
diff --git a/Engine/Source/Editor/Matinee/Classes/MatineeTrackColorPropHelper.h b/Engine/Source/Editor/Matinee/Classes/MatineeTrackColorPropHelper.h
--- a/Engine/Source/Editor/Matinee/Classes/MatineeTrackColorPropHelper.h
+++ b/Engine/Source/Editor/Matinee/Classes/MatineeTrackColorPropHelper.h
@@ -6,9 +6,7 @@
 UCLASS()
 class UMatineeTrackColorPropHelper : public UMatineeTrackVectorPropHelper
 {
-	GENERATED_BODY()
-public:
-	UMatineeTrackColorPropHelper(const FObjectInitializer& ObjectInitializer = FObjectInitializer::Get());
+	GENERATED_UCLASS_BODY()
 
 public:
 
diff --git a/Engine/Source/Editor/Matinee/Classes/MatineeTrackEventHelper.h b/Engine/Source/Editor/Matinee/Classes/MatineeTrackEventHelper.h
--- a/Engine/Source/Editor/Matinee/Classes/MatineeTrackEventHelper.h
+++ b/Engine/Source/Editor/Matinee/Classes/MatineeTrackEventHelper.h
@@ -8,9 +8,7 @@
 UCLASS()
 class UMatineeTrackEventHelper : public UInterpTrackHelper
 {
-	GENERATED_BODY()
-public:
-	UMatineeTrackEventHelper(const FObjectInitializer& ObjectInitializer = FObjectInitializer::Get());
+	GENERATED_UCLASS_BODY()
 
 	void OnAddKeyTextEntry(const FText& ChosenText, ETextCommit::Type CommitInfo, IMatineeBase* Matinee, UInterpTrack* Track);
 
diff --git a/Engine/Source/Editor/Matinee/Classes/MatineeTrackFloatPropHelper.h b/Engine/Source/Editor/Matinee/Classes/MatineeTrackFloatPropHelper.h
--- a/Engine/Source/Editor/Matinee/Classes/MatineeTrackFloatPropHelper.h
+++ b/Engine/Source/Editor/Matinee/Classes/MatineeTrackFloatPropHelper.h
@@ -10,9 +10,7 @@
 UCLASS()
 class UMatineeTrackFloatPropHelper : public UInterpTrackHelper
 {
-	GENERATED_BODY()
-public:
-	UMatineeTrackFloatPropHelper(const FObjectInitializer& ObjectInitializer = FObjectInitializer::Get());
+	GENERATED_UCLASS_BODY()
 
 	void OnCreateTrackTextEntry(const FString& ChosenText, TSharedRef<SWindow> Window, FString* OutputString);
 
diff --git a/Engine/Source/Editor/Matinee/Classes/MatineeTrackParticleReplayHelper.h b/Engine/Source/Editor/Matinee/Classes/MatineeTrackParticleReplayHelper.h
--- a/Engine/Source/Editor/Matinee/Classes/MatineeTrackParticleReplayHelper.h
+++ b/Engine/Source/Editor/Matinee/Classes/MatineeTrackParticleReplayHelper.h
@@ -8,9 +8,7 @@
 UCLASS()
 class UMatineeTrackParticleReplayHelper : public UInterpTrackHelper
 {
-	GENERATED_BODY()
-public:
-	UMatineeTrackParticleReplayHelper(const FObjectInitializer& ObjectInitializer = FObjectInitializer::Get());
+	GENERATED_UCLASS_BODY()
 
 public:
 
diff --git a/Engine/Source/Editor/Matinee/Classes/MatineeTrackSoundHelper.h b/Engine/Source/Editor/Matinee/Classes/MatineeTrackSoundHelper.h
--- a/Engine/Source/Editor/Matinee/Classes/MatineeTrackSoundHelper.h
+++ b/Engine/Source/Editor/Matinee/Classes/MatineeTrackSoundHelper.h
@@ -8,9 +8,7 @@
 UCLASS()
 class UMatineeTrackSoundHelper : public UInterpTrackHelper
 {
-	GENERATED_BODY()
-public:
-	UMatineeTrackSoundHelper(const FObjectInitializer& ObjectInitializer = FObjectInitializer::Get());
+	GENERATED_UCLASS_BODY()
 
 public:
 
diff --git a/Engine/Source/Editor/Matinee/Classes/MatineeTrackVectorPropHelper.h b/Engine/Source/Editor/Matinee/Classes/MatineeTrackVectorPropHelper.h
--- a/Engine/Source/Editor/Matinee/Classes/MatineeTrackVectorPropHelper.h
+++ b/Engine/Source/Editor/Matinee/Classes/MatineeTrackVectorPropHelper.h
@@ -8,9 +8,7 @@
 UCLASS()
 class UMatineeTrackVectorPropHelper : public UInterpTrackHelper
 {
-	GENERATED_BODY()
-public:
-	UMatineeTrackVectorPropHelper(const FObjectInitializer& ObjectInitializer = FObjectInitializer::Get());
+	GENERATED_UCLASS_BODY()
 
 	/**
 	 * Pops up a dialog letting user choose between a set of properties, then checks to see if that property has been bound to yet.
diff --git a/Engine/Source/Editor/Matinee/Classes/MatineeTrackVisibilityHelper.h b/Engine/Source/Editor/Matinee/Classes/MatineeTrackVisibilityHelper.h
--- a/Engine/Source/Editor/Matinee/Classes/MatineeTrackVisibilityHelper.h
+++ b/Engine/Source/Editor/Matinee/Classes/MatineeTrackVisibilityHelper.h
@@ -7,9 +7,7 @@
 UCLASS()
 class UMatineeTrackVisibilityHelper : public UInterpTrackHelper
 {
-	GENERATED_BODY()
-public:
-	UMatineeTrackVisibilityHelper(const FObjectInitializer& ObjectInitializer = FObjectInitializer::Get());
+	GENERATED_UCLASS_BODY()
 
 	void OnAddKeyTextEntry(const FString& ChosenText, IMatineeBase* Matinee, UInterpTrack* Track);
 
diff --git a/Engine/Source/Editor/NewsFeed/Private/Implementation/NewsFeedSettings.h b/Engine/Source/Editor/NewsFeed/Private/Implementation/NewsFeedSettings.h
--- a/Engine/Source/Editor/NewsFeed/Private/Implementation/NewsFeedSettings.h
+++ b/Engine/Source/Editor/NewsFeed/Private/Implementation/NewsFeedSettings.h
@@ -29,9 +29,7 @@ UCLASS(config=EditorGameAgnostic)
 class UNewsFeedSettings
 	: public UObject
 {
-	GENERATED_BODY()
-public:
-	UNewsFeedSettings(const FObjectInitializer& ObjectInitializer = FObjectInitializer::Get());
+	GENERATED_UCLASS_BODY()
 
 public:
 
diff --git a/Engine/Source/Editor/OnlineBlueprintSupport/Classes/K2Node_InAppPurchase.h b/Engine/Source/Editor/OnlineBlueprintSupport/Classes/K2Node_InAppPurchase.h
--- a/Engine/Source/Editor/OnlineBlueprintSupport/Classes/K2Node_InAppPurchase.h
+++ b/Engine/Source/Editor/OnlineBlueprintSupport/Classes/K2Node_InAppPurchase.h
@@ -7,7 +7,5 @@
 UCLASS()
 class ONLINEBLUEPRINTSUPPORT_API UK2Node_InAppPurchase : public UK2Node_BaseAsyncTask
 {
-	GENERATED_BODY()
-public:
-	UK2Node_InAppPurchase(const FObjectInitializer& ObjectInitializer = FObjectInitializer::Get());
+	GENERATED_UCLASS_BODY()
 };
diff --git a/Engine/Source/Editor/OnlineBlueprintSupport/Classes/K2Node_LeaderboardFlush.h b/Engine/Source/Editor/OnlineBlueprintSupport/Classes/K2Node_LeaderboardFlush.h
--- a/Engine/Source/Editor/OnlineBlueprintSupport/Classes/K2Node_LeaderboardFlush.h
+++ b/Engine/Source/Editor/OnlineBlueprintSupport/Classes/K2Node_LeaderboardFlush.h
@@ -7,9 +7,7 @@
 UCLASS()
 class ONLINEBLUEPRINTSUPPORT_API UK2Node_LeaderboardFlush : public UK2Node_BaseAsyncTask
 {
-	GENERATED_BODY()
-public:
-	UK2Node_LeaderboardFlush(const FObjectInitializer& ObjectInitializer = FObjectInitializer::Get());
+	GENERATED_UCLASS_BODY()
 
 	// UEdGraphNode interface
 	virtual FText GetTooltipText() const override;
diff --git a/Engine/Source/Editor/OnlineBlueprintSupport/Classes/K2Node_LeaderboardQuery.h b/Engine/Source/Editor/OnlineBlueprintSupport/Classes/K2Node_LeaderboardQuery.h
--- a/Engine/Source/Editor/OnlineBlueprintSupport/Classes/K2Node_LeaderboardQuery.h
+++ b/Engine/Source/Editor/OnlineBlueprintSupport/Classes/K2Node_LeaderboardQuery.h
@@ -7,7 +7,5 @@
 UCLASS()
 class ONLINEBLUEPRINTSUPPORT_API UK2Node_LeaderboardQuery : public UK2Node_BaseAsyncTask
 {
-	GENERATED_BODY()
-public:
-	UK2Node_LeaderboardQuery(const FObjectInitializer& ObjectInitializer = FObjectInitializer::Get());
+	GENERATED_UCLASS_BODY()
 };
diff --git a/Engine/Source/Editor/PhAT/Classes/PhATEdSkeletalMeshComponent.h b/Engine/Source/Editor/PhAT/Classes/PhATEdSkeletalMeshComponent.h
--- a/Engine/Source/Editor/PhAT/Classes/PhATEdSkeletalMeshComponent.h
+++ b/Engine/Source/Editor/PhAT/Classes/PhATEdSkeletalMeshComponent.h
@@ -8,9 +8,7 @@
 UCLASS()
 class UPhATEdSkeletalMeshComponent : public USkeletalMeshComponent
 {
-	GENERATED_BODY()
-public:
-	UPhATEdSkeletalMeshComponent(const FObjectInitializer& ObjectInitializer = FObjectInitializer::Get());
+	GENERATED_UCLASS_BODY()
 
 	/** Data and methods shared across multiple classes */
 	class FPhATSharedData* SharedData;
diff --git a/Engine/Source/Editor/ReferenceViewer/Classes/EdGraphNode_Reference.h b/Engine/Source/Editor/ReferenceViewer/Classes/EdGraphNode_Reference.h
--- a/Engine/Source/Editor/ReferenceViewer/Classes/EdGraphNode_Reference.h
+++ b/Engine/Source/Editor/ReferenceViewer/Classes/EdGraphNode_Reference.h
@@ -7,9 +7,7 @@
 UCLASS()
 class UEdGraphNode_Reference : public UEdGraphNode
 {
-	GENERATED_BODY()
-public:
-	UEdGraphNode_Reference(const FObjectInitializer& ObjectInitializer = FObjectInitializer::Get());
+	GENERATED_UCLASS_BODY()
 
 	virtual void SetupReferenceNode(const FIntPoint& NodeLoc, const TArray<FName>& NewPackageNames, const FAssetData& InAssetData);
 	virtual void SetReferenceNodeCollapsed(const FIntPoint& NodeLoc, int32 InNumReferencesExceedingMax);
diff --git a/Engine/Source/Editor/ReferenceViewer/Classes/EdGraph_ReferenceViewer.h b/Engine/Source/Editor/ReferenceViewer/Classes/EdGraph_ReferenceViewer.h
--- a/Engine/Source/Editor/ReferenceViewer/Classes/EdGraph_ReferenceViewer.h
+++ b/Engine/Source/Editor/ReferenceViewer/Classes/EdGraph_ReferenceViewer.h
@@ -7,9 +7,7 @@
 UCLASS()
 class UEdGraph_ReferenceViewer : public UEdGraph
 {
-	GENERATED_BODY()
-public:
-	UEdGraph_ReferenceViewer(const FObjectInitializer& ObjectInitializer = FObjectInitializer::Get());
+	GENERATED_UCLASS_BODY()
 
 public:
 	// UObject implementation
diff --git a/Engine/Source/Editor/ReferenceViewer/Classes/ReferenceViewerSchema.h b/Engine/Source/Editor/ReferenceViewer/Classes/ReferenceViewerSchema.h
--- a/Engine/Source/Editor/ReferenceViewer/Classes/ReferenceViewerSchema.h
+++ b/Engine/Source/Editor/ReferenceViewer/Classes/ReferenceViewerSchema.h
@@ -6,9 +6,7 @@
 UCLASS()
 class UReferenceViewerSchema : public UEdGraphSchema
 {
-	GENERATED_BODY()
-public:
-	UReferenceViewerSchema(const FObjectInitializer& ObjectInitializer = FObjectInitializer::Get());
+	GENERATED_UCLASS_BODY()
 
 public:
 	// UEdGraphSchema interface
diff --git a/Engine/Source/Editor/SceneOutliner/Private/SceneOutlinerSettings.h b/Engine/Source/Editor/SceneOutliner/Private/SceneOutlinerSettings.h
--- a/Engine/Source/Editor/SceneOutliner/Private/SceneOutlinerSettings.h
+++ b/Engine/Source/Editor/SceneOutliner/Private/SceneOutlinerSettings.h
@@ -12,9 +12,7 @@ UCLASS(config=EditorUserSettings)
 class USceneOutlinerSettings
 	: public UObject
 {
-	GENERATED_BODY()
-public:
-	USceneOutlinerSettings(const FObjectInitializer& ObjectInitializer = FObjectInitializer::Get());
+	GENERATED_UCLASS_BODY()
 
 	/** True when the Scene Outliner is hiding temporary/run-time Actors */
 	UPROPERTY(config)
diff --git a/Engine/Source/Editor/Sequencer/Private/SequencerSettings.h b/Engine/Source/Editor/Sequencer/Private/SequencerSettings.h
--- a/Engine/Source/Editor/Sequencer/Private/SequencerSettings.h
+++ b/Engine/Source/Editor/Sequencer/Private/SequencerSettings.h
@@ -23,10 +23,8 @@ namespace ESequencerCurveVisibility
 UCLASS(config=EditorUserSettings)
 class USequencerSettings : public UObject
 {
-	GENERATED_BODY()
-public:
-	USequencerSettings(const FObjectInitializer& ObjectInitializer = FObjectInitializer::Get());
 public:
+	GENERATED_UCLASS_BODY()
 
 	DECLARE_MULTICAST_DELEGATE( FOnCurveVisibilityChanged );
 
@@ -142,4 +140,4 @@ protected:
 	TEnumAsByte<ESequencerCurveVisibility::Type> CurveVisibility;
 
 	FOnCurveVisibilityChanged OnCurveVisibilityChanged;
-};
+};
\ No newline at end of file
diff --git a/Engine/Source/Editor/StatsViewer/Classes/CookerStats.h b/Engine/Source/Editor/StatsViewer/Classes/CookerStats.h
--- a/Engine/Source/Editor/StatsViewer/Classes/CookerStats.h
+++ b/Engine/Source/Editor/StatsViewer/Classes/CookerStats.h
@@ -22,9 +22,7 @@ UCLASS(Transient, MinimalAPI, meta=(DisplayName="Cooker Stats", ObjectSetType="E
 class UCookerStats
 	: public UObject
 {
-	GENERATED_BODY()
-public:
-	STATSVIEWER_API UCookerStats(const FObjectInitializer& ObjectInitializer = FObjectInitializer::Get());
+	GENERATED_UCLASS_BODY()
 
 	/** The assets contained in the file. */
 	UPROPERTY(VisibleAnywhere, AssetRegistrySearchable, Category="Stats", meta=(DisplayName="Asset(s)", ColumnWidth="100"))
diff --git a/Engine/Source/Editor/StatsViewer/Classes/LightingBuildInfo.h b/Engine/Source/Editor/StatsViewer/Classes/LightingBuildInfo.h
--- a/Engine/Source/Editor/StatsViewer/Classes/LightingBuildInfo.h
+++ b/Engine/Source/Editor/StatsViewer/Classes/LightingBuildInfo.h
@@ -15,9 +15,7 @@ enum ELightingBuildInfoObjectSets
 UCLASS(Transient, MinimalAPI, meta=( DisplayName = "Lighting Build Info", ObjectSetType = "ELightingBuildInfoObjectSets" ) )
 class ULightingBuildInfo : public UObject
 {
-	GENERATED_BODY()
-public:
-	STATSVIEWER_API ULightingBuildInfo(const FObjectInitializer& ObjectInitializer = FObjectInitializer::Get());
+	GENERATED_UCLASS_BODY()
 
 	/** The actor and/or object that is related to this info. */
 	UPROPERTY(VisibleAnywhere, AssetRegistrySearchable, Category="Stats", meta=( ColumnWidth = "150" ) )
@@ -59,4 +57,4 @@ private:
 
 	/** Update internal strings */
 	void UpdateNames();
-};
+};
\ No newline at end of file
diff --git a/Engine/Source/Editor/StatsViewer/Classes/PrimitiveStats.h b/Engine/Source/Editor/StatsViewer/Classes/PrimitiveStats.h
--- a/Engine/Source/Editor/StatsViewer/Classes/PrimitiveStats.h
+++ b/Engine/Source/Editor/StatsViewer/Classes/PrimitiveStats.h
@@ -17,9 +17,7 @@ enum EPrimitiveObjectSets
 UCLASS(Transient, MinimalAPI, meta=( DisplayName = "Primitive Stats", ObjectSetType = "EPrimitiveObjectSets" ) )
 class UPrimitiveStats : public UObject
 {
-	GENERATED_BODY()
-public:
-	STATSVIEWER_API UPrimitiveStats(const FObjectInitializer& ObjectInitializer = FObjectInitializer::Get());
+	GENERATED_UCLASS_BODY()
 
 	/** Resource (e.g. UStaticMesh, USkeletalMesh, UModelComponent, UTerrainComponent, etc */
 	UPROPERTY(VisibleAnywhere, AssetRegistrySearchable, Category="Stats", meta=( DisplayName = "Object", ColumnWidth = "200" ) )
@@ -105,4 +103,4 @@ public:
 
 	/** Update internal strings */
 	void UpdateNames();
-};
+};
\ No newline at end of file
diff --git a/Engine/Source/Editor/StatsViewer/Classes/StaticMeshLightingInfo.h b/Engine/Source/Editor/StatsViewer/Classes/StaticMeshLightingInfo.h
--- a/Engine/Source/Editor/StatsViewer/Classes/StaticMeshLightingInfo.h
+++ b/Engine/Source/Editor/StatsViewer/Classes/StaticMeshLightingInfo.h
@@ -17,9 +17,7 @@ enum EStaticMeshLightingInfoObjectSets
 UCLASS(Transient, MinimalAPI, meta=( DisplayName = "Static Mesh Lighting Info", ObjectSetType = "EStaticMeshLightingInfoObjectSets" ) )
 class UStaticMeshLightingInfo : public UObject
 {
-	GENERATED_BODY()
-public:
-	STATSVIEWER_API UStaticMeshLightingInfo(const FObjectInitializer& ObjectInitializer = FObjectInitializer::Get());
+	GENERATED_UCLASS_BODY()
 
 	/** The actor that is related to this error/warning. */
 	UPROPERTY(VisibleAnywhere, AssetRegistrySearchable, Category="Stats", meta=( DisplayName = "Actor", ColumnWidth = "200" ) )
@@ -80,4 +78,4 @@ public:
 
 	/** Update internal strings */
 	void UpdateNames();
-};
+};
\ No newline at end of file
diff --git a/Engine/Source/Editor/StatsViewer/Classes/TextureStats.h b/Engine/Source/Editor/StatsViewer/Classes/TextureStats.h
--- a/Engine/Source/Editor/StatsViewer/Classes/TextureStats.h
+++ b/Engine/Source/Editor/StatsViewer/Classes/TextureStats.h
@@ -24,9 +24,7 @@ enum ETextureObjectSets
 UCLASS(Transient, MinimalAPI, meta=( DisplayName = "Texture Stats", ObjectSetType = "ETextureObjectSets" ) )
 class UTextureStats : public UObject
 {
-	GENERATED_BODY()
-public:
-	STATSVIEWER_API UTextureStats(const FObjectInitializer& ObjectInitializer = FObjectInitializer::Get());
+	GENERATED_UCLASS_BODY()
 
 	/** Texture - click to go to asset */
 	UPROPERTY(VisibleAnywhere, AssetRegistrySearchable, Category="Stats", meta=( DisplayName = "Name", ColumnWidth = "100" ) )
@@ -79,4 +77,4 @@ public:
 	/** Texture path without the name "package.[group.]" */
 	UPROPERTY(VisibleAnywhere, AssetRegistrySearchable, Category="Stats", meta=( ColumnWidth = "200" ) )
 	FString Path;
-};
+};
\ No newline at end of file
diff --git a/Engine/Source/Editor/SwarmInterface/Private/SwarmMessages.h b/Engine/Source/Editor/SwarmInterface/Private/SwarmMessages.h
--- a/Engine/Source/Editor/SwarmInterface/Private/SwarmMessages.h
+++ b/Engine/Source/Editor/SwarmInterface/Private/SwarmMessages.h
@@ -9,7 +9,7 @@
 USTRUCT()
 struct FSwarmPingMessage
 {
-	GENERATED_BODY()
+	GENERATED_USTRUCT_BODY()
 };
 
 template<>
@@ -25,7 +25,7 @@ struct TStructOpsTypeTraits<FSwarmPingMessage> : public TStructOpsTypeTraitsBase
 USTRUCT()
 struct FSwarmPongMessage
 {
-	GENERATED_BODY()
+	GENERATED_USTRUCT_BODY()
 
 	UPROPERTY()
 	bool			bIsEditor;
@@ -55,7 +55,7 @@ struct TStructOpsTypeTraits<FSwarmPongMessage> : public TStructOpsTypeTraitsBase
 USTRUCT()
 struct FSwarmInfoMessage
 {
-	GENERATED_BODY()
+	GENERATED_USTRUCT_BODY()
 
 	UPROPERTY()
 	FString TextMessage;
@@ -81,7 +81,7 @@ struct TStructOpsTypeTraits<FSwarmInfoMessage> : public TStructOpsTypeTraitsBase
 USTRUCT()
 struct FSwarmAlertMessage
 {
-	GENERATED_BODY()
+	GENERATED_USTRUCT_BODY()
 
 	/** The Job Guid */
 	UPROPERTY()
@@ -124,7 +124,7 @@ struct TStructOpsTypeTraits<FSwarmAlertMessage> : public TStructOpsTypeTraitsBas
 USTRUCT()
 struct FSwarmTimingMessage
 {
-	GENERATED_BODY()
+	GENERATED_USTRUCT_BODY()
 
 	/** State that the distributed job is transitioning to */
 	UPROPERTY()
@@ -155,7 +155,7 @@ struct TStructOpsTypeTraits<FSwarmTimingMessage> : public TStructOpsTypeTraitsBa
 USTRUCT()
 struct FSwarmTaskRequestMessage
 {
-	GENERATED_BODY()
+	GENERATED_USTRUCT_BODY()
 
 	FSwarmTaskRequestMessage() {}
 };
@@ -173,7 +173,7 @@ struct TStructOpsTypeTraits<FSwarmTaskRequestMessage> : public TStructOpsTypeTra
 USTRUCT()
 struct FSwarmTaskRequestReleaseMessage
 {
-	GENERATED_BODY()
+	GENERATED_USTRUCT_BODY()
 
 	FSwarmTaskRequestReleaseMessage() {}
 };
@@ -191,7 +191,7 @@ struct TStructOpsTypeTraits<FSwarmTaskRequestReleaseMessage> : public TStructOps
 USTRUCT()
 struct FSwarmTaskRequestReservationMessage
 {
-	GENERATED_BODY()
+	GENERATED_USTRUCT_BODY()
 
 	FSwarmTaskRequestReservationMessage() {}
 };
@@ -209,7 +209,7 @@ struct TStructOpsTypeTraits<FSwarmTaskRequestReservationMessage> : public TStruc
 USTRUCT()
 struct FSwarmTaskRequestSpecificationMessage
 {
-	GENERATED_BODY()
+	GENERATED_USTRUCT_BODY()
 
 	/** The GUID used for identifying the Task being referred to */
 	UPROPERTY()
@@ -250,7 +250,7 @@ struct TStructOpsTypeTraits<FSwarmTaskRequestSpecificationMessage> : public TStr
 USTRUCT()
 struct FSwarmJobStateMessage
 {
-	GENERATED_BODY()
+	GENERATED_USTRUCT_BODY()
 
 	/** The Job GUID used for identifying the Job */
 	UPROPERTY()
@@ -291,7 +291,7 @@ struct TStructOpsTypeTraits<FSwarmJobStateMessage> : public TStructOpsTypeTraits
 USTRUCT()
 struct FSwarmTaskStateMessage
 {
-	GENERATED_BODY()
+	GENERATED_USTRUCT_BODY()
 
 	/** The Task GUID used for identifying the Task */
 	UPROPERTY()
@@ -332,7 +332,7 @@ struct TStructOpsTypeTraits<FSwarmTaskStateMessage> : public TStructOpsTypeTrait
 USTRUCT()
 struct FSwarmQuitMessage
 {
-	GENERATED_BODY()
+	GENERATED_USTRUCT_BODY()
 
 	FSwarmQuitMessage() {}
 };
diff --git a/Engine/Source/Editor/TextureEditor/Classes/TextureEditorSettings.h b/Engine/Source/Editor/TextureEditor/Classes/TextureEditorSettings.h
--- a/Engine/Source/Editor/TextureEditor/Classes/TextureEditorSettings.h
+++ b/Engine/Source/Editor/TextureEditor/Classes/TextureEditorSettings.h
@@ -25,9 +25,7 @@ UCLASS(config=EditorUserSettings)
 class TEXTUREEDITOR_API UTextureEditorSettings
 	: public UObject
 {
-	GENERATED_BODY()
-public:
-	UTextureEditorSettings(const FObjectInitializer& ObjectInitializer = FObjectInitializer::Get());
+	GENERATED_UCLASS_BODY()
 
 public:
 
diff --git a/Engine/Source/Editor/TranslationEditor/Private/TranslationUnit.h b/Engine/Source/Editor/TranslationEditor/Private/TranslationUnit.h
--- a/Engine/Source/Editor/TranslationEditor/Private/TranslationUnit.h
+++ b/Engine/Source/Editor/TranslationEditor/Private/TranslationUnit.h
@@ -6,7 +6,7 @@
 USTRUCT()
 struct FTranslationChange
 {
-	GENERATED_BODY()
+	GENERATED_USTRUCT_BODY()
 
 public:
 
@@ -30,7 +30,7 @@ public:
 USTRUCT()
 struct FTranslationContextInfo
 {
-	GENERATED_BODY()
+	GENERATED_USTRUCT_BODY()
 
 public:
 
@@ -50,9 +50,7 @@ public:
 UCLASS(hidecategories = Object, MinimalAPI)
 class UTranslationUnit : public UObject
 {
-	GENERATED_BODY()
-public:
-	TRANSLATIONEDITOR_API UTranslationUnit(const FObjectInitializer& ObjectInitializer = FObjectInitializer::Get());
+	GENERATED_UCLASS_BODY()
 
 public:
 
@@ -104,4 +102,4 @@ private:
 
 	// Holds an event delegate that is executed when a property has changed.
 	FTranslationUnitPropertyChangedEvent TranslationUnitPropertyChangedEvent;
-};
+};
\ No newline at end of file
diff --git a/Engine/Source/Editor/UMGEditor/Classes/WidgetBlueprintFactory.h b/Engine/Source/Editor/UMGEditor/Classes/WidgetBlueprintFactory.h
--- a/Engine/Source/Editor/UMGEditor/Classes/WidgetBlueprintFactory.h
+++ b/Engine/Source/Editor/UMGEditor/Classes/WidgetBlueprintFactory.h
@@ -7,9 +7,7 @@
 UCLASS(HideCategories=Object, MinimalAPI)
 class UWidgetBlueprintFactory : public UFactory
 {
-	GENERATED_BODY()
-public:
-	UMGEDITOR_API UWidgetBlueprintFactory(const FObjectInitializer& ObjectInitializer = FObjectInitializer::Get());
+	GENERATED_UCLASS_BODY()
 
 	// The type of blueprint that will be created
 	UPROPERTY(EditAnywhere, Category=WidgetBlueprintFactory)
diff --git a/Engine/Source/Editor/UMGEditor/Private/Nodes/K2Node_CreateDragDropOperation.h b/Engine/Source/Editor/UMGEditor/Private/Nodes/K2Node_CreateDragDropOperation.h
--- a/Engine/Source/Editor/UMGEditor/Private/Nodes/K2Node_CreateDragDropOperation.h
+++ b/Engine/Source/Editor/UMGEditor/Private/Nodes/K2Node_CreateDragDropOperation.h
@@ -9,9 +9,7 @@
 UCLASS()
 class UMGEDITOR_API UK2Node_CreateDragDropOperation : public UK2Node_ConstructObjectFromClass
 {
-	GENERATED_BODY()
-public:
-	UK2Node_CreateDragDropOperation(const FObjectInitializer& ObjectInitializer = FObjectInitializer::Get());
+	GENERATED_UCLASS_BODY()
 
 	// Begin UEdGraphNode interface.
 	virtual void AllocateDefaultPins() override;
diff --git a/Engine/Source/Editor/UMGEditor/Private/Nodes/K2Node_CreateWidget.h b/Engine/Source/Editor/UMGEditor/Private/Nodes/K2Node_CreateWidget.h
--- a/Engine/Source/Editor/UMGEditor/Private/Nodes/K2Node_CreateWidget.h
+++ b/Engine/Source/Editor/UMGEditor/Private/Nodes/K2Node_CreateWidget.h
@@ -9,9 +9,7 @@
 UCLASS()
 class UMGEDITOR_API UK2Node_CreateWidget : public UK2Node_ConstructObjectFromClass
 {
-	GENERATED_BODY()
-public:
-	UK2Node_CreateWidget(const FObjectInitializer& ObjectInitializer = FObjectInitializer::Get());
+	GENERATED_UCLASS_BODY()
 
 	// Begin UEdGraphNode interface.
 	virtual void AllocateDefaultPins() override;
diff --git a/Engine/Source/Editor/UMGEditor/Public/Settings/WidgetDesignerSettings.h b/Engine/Source/Editor/UMGEditor/Public/Settings/WidgetDesignerSettings.h
--- a/Engine/Source/Editor/UMGEditor/Public/Settings/WidgetDesignerSettings.h
+++ b/Engine/Source/Editor/UMGEditor/Public/Settings/WidgetDesignerSettings.h
@@ -11,9 +11,7 @@ UCLASS(config=EditorUserSettings)
 class UMGEDITOR_API UWidgetDesignerSettings
 	: public UObject
 {
-	GENERATED_BODY()
-public:
-	UWidgetDesignerSettings(const FObjectInitializer& ObjectInitializer = FObjectInitializer::Get());
+	GENERATED_UCLASS_BODY()
 
 	/** If enabled, actor positions will snap to the grid. */
 	UPROPERTY(EditAnywhere, config, Category=GridSnapping, meta=(DisplayName = "Enable Grid Snapping"))
diff --git a/Engine/Source/Editor/UMGEditor/Public/WidgetBlueprint.h b/Engine/Source/Editor/UMGEditor/Public/WidgetBlueprint.h
--- a/Engine/Source/Editor/UMGEditor/Public/WidgetBlueprint.h
+++ b/Engine/Source/Editor/UMGEditor/Public/WidgetBlueprint.h
@@ -15,7 +15,7 @@ class UMovieScene;
 USTRUCT()
 struct UMGEDITOR_API FEditorPropertyPathSegment
 {
-	GENERATED_BODY()
+	GENERATED_USTRUCT_BODY()
 
 public:
 	FEditorPropertyPathSegment();
@@ -59,7 +59,7 @@ private:
 USTRUCT()
 struct UMGEDITOR_API FEditorPropertyPath
 {
-	GENERATED_BODY()
+	GENERATED_USTRUCT_BODY()
 
 public:
 
@@ -96,7 +96,7 @@ public:
 USTRUCT()
 struct UMGEDITOR_API FDelegateEditorBinding
 {
-	GENERATED_BODY()
+	GENERATED_USTRUCT_BODY()
 
 	/** The member widget the binding is on, must be a direct variable of the UUserWidget. */
 	UPROPERTY()
@@ -142,7 +142,7 @@ struct UMGEDITOR_API FDelegateEditorBinding
 USTRUCT()
 struct FWidgetAnimation_DEPRECATED
 {
-	GENERATED_BODY()
+	GENERATED_USTRUCT_BODY()
 
 	UPROPERTY()
 	UMovieScene* MovieScene;
@@ -170,9 +170,7 @@ struct TStructOpsTypeTraits<FWidgetAnimation_DEPRECATED> : public TStructOpsType
 UCLASS(BlueprintType)
 class UMGEDITOR_API UWidgetBlueprint : public UBlueprint
 {
-	GENERATED_BODY()
-public:
-	UWidgetBlueprint(const FObjectInitializer& ObjectInitializer = FObjectInitializer::Get());
+	GENERATED_UCLASS_BODY()
 
 public:
 
@@ -224,4 +222,4 @@ public:
 	// End of UBlueprint interface
 
 	static bool ValidateGeneratedClass(const UClass* InClass);
-};
+};
\ No newline at end of file
diff --git a/Engine/Source/Editor/UnrealEd/Classes/ActorFactories/ActorFactory.h b/Engine/Source/Editor/UnrealEd/Classes/ActorFactories/ActorFactory.h
--- a/Engine/Source/Editor/UnrealEd/Classes/ActorFactories/ActorFactory.h
+++ b/Engine/Source/Editor/UnrealEd/Classes/ActorFactories/ActorFactory.h
@@ -11,9 +11,7 @@ class ULevel;
 UCLASS(collapsecategories, hidecategories=Object, editinlinenew, config=Editor, abstract, transient)
 class UNREALED_API UActorFactory : public UObject
 {
-	GENERATED_BODY()
-public:
-	UActorFactory(const FObjectInitializer& ObjectInitializer = FObjectInitializer::Get());
+	GENERATED_UCLASS_BODY()
 
 	/** Name used as basis for 'New Actor' menu. */
 	UPROPERTY()
diff --git a/Engine/Source/Editor/UnrealEd/Classes/ActorFactories/ActorFactoryAmbientSound.h b/Engine/Source/Editor/UnrealEd/Classes/ActorFactories/ActorFactoryAmbientSound.h
--- a/Engine/Source/Editor/UnrealEd/Classes/ActorFactories/ActorFactoryAmbientSound.h
+++ b/Engine/Source/Editor/UnrealEd/Classes/ActorFactories/ActorFactoryAmbientSound.h
@@ -12,9 +12,7 @@
 UCLASS(MinimalAPI, config=Editor, collapsecategories, hidecategories=Object)
 class UActorFactoryAmbientSound : public UActorFactory
 {
-	GENERATED_BODY()
-public:
-	UNREALED_API UActorFactoryAmbientSound(const FObjectInitializer& ObjectInitializer = FObjectInitializer::Get());
+	GENERATED_UCLASS_BODY()
 
 	// Begin UActorFactory Interface
 	virtual void PostSpawnActor( UObject* Asset, AActor* NewActor ) override;
diff --git a/Engine/Source/Editor/UnrealEd/Classes/ActorFactories/ActorFactoryAtmosphericFog.h b/Engine/Source/Editor/UnrealEd/Classes/ActorFactories/ActorFactoryAtmosphericFog.h
--- a/Engine/Source/Editor/UnrealEd/Classes/ActorFactories/ActorFactoryAtmosphericFog.h
+++ b/Engine/Source/Editor/UnrealEd/Classes/ActorFactories/ActorFactoryAtmosphericFog.h
@@ -7,9 +7,7 @@
 UCLASS(MinimalAPI, config=Editor, collapsecategories, hidecategories=Object)
 class UActorFactoryAtmosphericFog : public UActorFactory
 {
-	GENERATED_BODY()
-public:
-	UNREALED_API UActorFactoryAtmosphericFog(const FObjectInitializer& ObjectInitializer = FObjectInitializer::Get());
+	GENERATED_UCLASS_BODY()
 };
 
 
diff --git a/Engine/Source/Editor/UnrealEd/Classes/ActorFactories/ActorFactoryBoxReflectionCapture.h b/Engine/Source/Editor/UnrealEd/Classes/ActorFactories/ActorFactoryBoxReflectionCapture.h
--- a/Engine/Source/Editor/UnrealEd/Classes/ActorFactories/ActorFactoryBoxReflectionCapture.h
+++ b/Engine/Source/Editor/UnrealEd/Classes/ActorFactories/ActorFactoryBoxReflectionCapture.h
@@ -7,9 +7,7 @@
 UCLASS(MinimalAPI, config=Editor, collapsecategories, hidecategories=Object)
 class UActorFactoryBoxReflectionCapture : public UActorFactory
 {
-	GENERATED_BODY()
-public:
-	UNREALED_API UActorFactoryBoxReflectionCapture(const FObjectInitializer& ObjectInitializer = FObjectInitializer::Get());
+	GENERATED_UCLASS_BODY()
 };
 
 
diff --git a/Engine/Source/Editor/UnrealEd/Classes/ActorFactories/ActorFactoryBoxVolume.h b/Engine/Source/Editor/UnrealEd/Classes/ActorFactories/ActorFactoryBoxVolume.h
--- a/Engine/Source/Editor/UnrealEd/Classes/ActorFactories/ActorFactoryBoxVolume.h
+++ b/Engine/Source/Editor/UnrealEd/Classes/ActorFactories/ActorFactoryBoxVolume.h
@@ -8,9 +8,7 @@
 UCLASS(MinimalAPI, config=Editor, collapsecategories, hidecategories=Object)
 class UActorFactoryBoxVolume : public UActorFactory
 {
-	GENERATED_BODY()
-public:
-	UNREALED_API UActorFactoryBoxVolume(const FObjectInitializer& ObjectInitializer = FObjectInitializer::Get());
+	GENERATED_UCLASS_BODY()
 
 	// Begin UActorFactory Interface
 	virtual bool CanCreateActorFrom( const FAssetData& AssetData, FText& OutErrorMsg ) override;
diff --git a/Engine/Source/Editor/UnrealEd/Classes/ActorFactories/ActorFactoryCameraActor.h b/Engine/Source/Editor/UnrealEd/Classes/ActorFactories/ActorFactoryCameraActor.h
--- a/Engine/Source/Editor/UnrealEd/Classes/ActorFactories/ActorFactoryCameraActor.h
+++ b/Engine/Source/Editor/UnrealEd/Classes/ActorFactories/ActorFactoryCameraActor.h
@@ -7,7 +7,5 @@
 UCLASS(MinimalAPI, config=Editor, collapsecategories, hidecategories=Object)
 class UActorFactoryCameraActor : public UActorFactory
 {
-	GENERATED_BODY()
-public:
-	UNREALED_API UActorFactoryCameraActor(const FObjectInitializer& ObjectInitializer = FObjectInitializer::Get());
+	GENERATED_UCLASS_BODY()
 };
diff --git a/Engine/Source/Editor/UnrealEd/Classes/ActorFactories/ActorFactoryCharacter.h b/Engine/Source/Editor/UnrealEd/Classes/ActorFactories/ActorFactoryCharacter.h
--- a/Engine/Source/Editor/UnrealEd/Classes/ActorFactories/ActorFactoryCharacter.h
+++ b/Engine/Source/Editor/UnrealEd/Classes/ActorFactories/ActorFactoryCharacter.h
@@ -7,9 +7,7 @@
 UCLASS(MinimalAPI,config=Editor)
 class UActorFactoryCharacter : public UActorFactory
 {
-	GENERATED_BODY()
-public:
-	UNREALED_API UActorFactoryCharacter(const FObjectInitializer& ObjectInitializer = FObjectInitializer::Get());
+	GENERATED_UCLASS_BODY()
 
 public:
 	virtual bool CanCreateActorFrom( const FAssetData& AssetData, FText& OutErrorMsg ) override;
diff --git a/Engine/Source/Editor/UnrealEd/Classes/ActorFactories/ActorFactoryClass.h b/Engine/Source/Editor/UnrealEd/Classes/ActorFactories/ActorFactoryClass.h
--- a/Engine/Source/Editor/UnrealEd/Classes/ActorFactories/ActorFactoryClass.h
+++ b/Engine/Source/Editor/UnrealEd/Classes/ActorFactories/ActorFactoryClass.h
@@ -8,9 +8,7 @@
 UCLASS(MinimalAPI, config=Editor, collapsecategories, hidecategories=Object)
 class UActorFactoryClass : public UActorFactory
 {
-	GENERATED_BODY()
-public:
-	UNREALED_API UActorFactoryClass(const FObjectInitializer& ObjectInitializer = FObjectInitializer::Get());
+	GENERATED_UCLASS_BODY()
 
 	// Begin UActorFactory Interface
 	virtual bool PreSpawnActor( UObject* Asset, FTransform& InOutLocation ) override;
diff --git a/Engine/Source/Editor/UnrealEd/Classes/ActorFactories/ActorFactoryCylinderVolume.h b/Engine/Source/Editor/UnrealEd/Classes/ActorFactories/ActorFactoryCylinderVolume.h
--- a/Engine/Source/Editor/UnrealEd/Classes/ActorFactories/ActorFactoryCylinderVolume.h
+++ b/Engine/Source/Editor/UnrealEd/Classes/ActorFactories/ActorFactoryCylinderVolume.h
@@ -8,9 +8,7 @@
 UCLASS(MinimalAPI, config=Editor, collapsecategories, hidecategories=Object)
 class UActorFactoryCylinderVolume : public UActorFactory
 {
-	GENERATED_BODY()
-public:
-	UNREALED_API UActorFactoryCylinderVolume(const FObjectInitializer& ObjectInitializer = FObjectInitializer::Get());
+	GENERATED_UCLASS_BODY()
 
 	// Begin UActorFactory Interface
 	virtual bool CanCreateActorFrom( const FAssetData& AssetData, FText& OutErrorMsg ) override;
diff --git a/Engine/Source/Editor/UnrealEd/Classes/ActorFactories/ActorFactoryDeferredDecal.h b/Engine/Source/Editor/UnrealEd/Classes/ActorFactories/ActorFactoryDeferredDecal.h
--- a/Engine/Source/Editor/UnrealEd/Classes/ActorFactories/ActorFactoryDeferredDecal.h
+++ b/Engine/Source/Editor/UnrealEd/Classes/ActorFactories/ActorFactoryDeferredDecal.h
@@ -8,10 +8,8 @@
 UCLASS(MinimalAPI, config=Editor)
 class UActorFactoryDeferredDecal : public UActorFactory
 {
-	GENERATED_BODY()
-public:
-	UNREALED_API UActorFactoryDeferredDecal(const FObjectInitializer& ObjectInitializer = FObjectInitializer::Get());
 public:
+	GENERATED_UCLASS_BODY()
 
 protected:
 	// Begin UActorFactory Interface
diff --git a/Engine/Source/Editor/UnrealEd/Classes/ActorFactories/ActorFactoryDestructible.h b/Engine/Source/Editor/UnrealEd/Classes/ActorFactories/ActorFactoryDestructible.h
--- a/Engine/Source/Editor/UnrealEd/Classes/ActorFactories/ActorFactoryDestructible.h
+++ b/Engine/Source/Editor/UnrealEd/Classes/ActorFactories/ActorFactoryDestructible.h
@@ -8,9 +8,7 @@
 UCLASS(MinimalAPI, config=Editor)
 class UActorFactoryDestructible : public UActorFactory
 {
-	GENERATED_BODY()
-public:
-	UNREALED_API UActorFactoryDestructible(const FObjectInitializer& ObjectInitializer = FObjectInitializer::Get());
+	GENERATED_UCLASS_BODY()
 
 	// Begin UActorFactory Interface
 	virtual void PostSpawnActor( UObject* Asset, AActor* NewActor ) override;
diff --git a/Engine/Source/Editor/UnrealEd/Classes/ActorFactories/ActorFactoryDirectionalLight.h b/Engine/Source/Editor/UnrealEd/Classes/ActorFactories/ActorFactoryDirectionalLight.h
--- a/Engine/Source/Editor/UnrealEd/Classes/ActorFactories/ActorFactoryDirectionalLight.h
+++ b/Engine/Source/Editor/UnrealEd/Classes/ActorFactories/ActorFactoryDirectionalLight.h
@@ -7,9 +7,7 @@
 UCLASS(MinimalAPI, config=Editor, collapsecategories, hidecategories=Object)
 class UActorFactoryDirectionalLight : public UActorFactory
 {
-	GENERATED_BODY()
-public:
-	UNREALED_API UActorFactoryDirectionalLight(const FObjectInitializer& ObjectInitializer = FObjectInitializer::Get());
+	GENERATED_UCLASS_BODY()
 
 };
 
diff --git a/Engine/Source/Editor/UnrealEd/Classes/ActorFactories/ActorFactoryEmitter.h b/Engine/Source/Editor/UnrealEd/Classes/ActorFactories/ActorFactoryEmitter.h
--- a/Engine/Source/Editor/UnrealEd/Classes/ActorFactories/ActorFactoryEmitter.h
+++ b/Engine/Source/Editor/UnrealEd/Classes/ActorFactories/ActorFactoryEmitter.h
@@ -8,9 +8,7 @@
 UCLASS(MinimalAPI, config=Editor, collapsecategories, hidecategories=Object)
 class UActorFactoryEmitter : public UActorFactory
 {
-	GENERATED_BODY()
-public:
-	UNREALED_API UActorFactoryEmitter(const FObjectInitializer& ObjectInitializer = FObjectInitializer::Get());
+	GENERATED_UCLASS_BODY()
 
 	// Begin UActorFactory Interface
 	virtual void PostSpawnActor( UObject* Asset, AActor* NewActor ) override;
diff --git a/Engine/Source/Editor/UnrealEd/Classes/ActorFactories/ActorFactoryEmptyActor.h b/Engine/Source/Editor/UnrealEd/Classes/ActorFactories/ActorFactoryEmptyActor.h
--- a/Engine/Source/Editor/UnrealEd/Classes/ActorFactories/ActorFactoryEmptyActor.h
+++ b/Engine/Source/Editor/UnrealEd/Classes/ActorFactories/ActorFactoryEmptyActor.h
@@ -7,9 +7,7 @@
 UCLASS(MinimalAPI,config=Editor)
 class UActorFactoryEmptyActor : public UActorFactory
 {
-	GENERATED_BODY()
-public:
-	UNREALED_API UActorFactoryEmptyActor(const FObjectInitializer& ObjectInitializer = FObjectInitializer::Get());
+	GENERATED_UCLASS_BODY()
 
 public:
 	virtual bool CanCreateActorFrom( const FAssetData& AssetData, FText& OutErrorMsg ) override;
diff --git a/Engine/Source/Editor/UnrealEd/Classes/ActorFactories/ActorFactoryExponentialHeightFog.h b/Engine/Source/Editor/UnrealEd/Classes/ActorFactories/ActorFactoryExponentialHeightFog.h
--- a/Engine/Source/Editor/UnrealEd/Classes/ActorFactories/ActorFactoryExponentialHeightFog.h
+++ b/Engine/Source/Editor/UnrealEd/Classes/ActorFactories/ActorFactoryExponentialHeightFog.h
@@ -7,9 +7,7 @@
 UCLASS(MinimalAPI, config=Editor, collapsecategories, hidecategories=Object)
 class UActorFactoryExponentialHeightFog : public UActorFactory
 {
-	GENERATED_BODY()
-public:
-	UNREALED_API UActorFactoryExponentialHeightFog(const FObjectInitializer& ObjectInitializer = FObjectInitializer::Get());
+	GENERATED_UCLASS_BODY()
 };
 
 
diff --git a/Engine/Source/Editor/UnrealEd/Classes/ActorFactories/ActorFactoryInteractiveFoliage.h b/Engine/Source/Editor/UnrealEd/Classes/ActorFactories/ActorFactoryInteractiveFoliage.h
--- a/Engine/Source/Editor/UnrealEd/Classes/ActorFactories/ActorFactoryInteractiveFoliage.h
+++ b/Engine/Source/Editor/UnrealEd/Classes/ActorFactories/ActorFactoryInteractiveFoliage.h
@@ -7,9 +7,7 @@
 UCLASS(MinimalAPI, config=Editor)
 class UActorFactoryInteractiveFoliage : public UActorFactoryStaticMesh
 {
-	GENERATED_BODY()
-public:
-	UNREALED_API UActorFactoryInteractiveFoliage(const FObjectInitializer& ObjectInitializer = FObjectInitializer::Get());
+	GENERATED_UCLASS_BODY()
 
 
 	// Begin UActorFactory Interface
diff --git a/Engine/Source/Editor/UnrealEd/Classes/ActorFactories/ActorFactoryMatineeActor.h b/Engine/Source/Editor/UnrealEd/Classes/ActorFactories/ActorFactoryMatineeActor.h
--- a/Engine/Source/Editor/UnrealEd/Classes/ActorFactories/ActorFactoryMatineeActor.h
+++ b/Engine/Source/Editor/UnrealEd/Classes/ActorFactories/ActorFactoryMatineeActor.h
@@ -8,9 +8,7 @@
 UCLASS(MinimalAPI, config=Editor)
 class UActorFactoryMatineeActor : public UActorFactory
 {
-	GENERATED_BODY()
-public:
-	UNREALED_API UActorFactoryMatineeActor(const FObjectInitializer& ObjectInitializer = FObjectInitializer::Get());
+	GENERATED_UCLASS_BODY()
 
 	// Begin UActorFactory Interface
 	virtual void PostSpawnActor( UObject* Asset, AActor* NewActor ) override;
diff --git a/Engine/Source/Editor/UnrealEd/Classes/ActorFactories/ActorFactoryNote.h b/Engine/Source/Editor/UnrealEd/Classes/ActorFactories/ActorFactoryNote.h
--- a/Engine/Source/Editor/UnrealEd/Classes/ActorFactories/ActorFactoryNote.h
+++ b/Engine/Source/Editor/UnrealEd/Classes/ActorFactories/ActorFactoryNote.h
@@ -6,8 +6,6 @@
 UCLASS( MinimalAPI, config = Editor, collapsecategories, hidecategories = Object )
 class UActorFactoryNote : public UActorFactory
 {
-	GENERATED_BODY()
-public:
-	UNREALED_API UActorFactoryNote(const FObjectInitializer& ObjectInitializer = FObjectInitializer::Get());
+	GENERATED_UCLASS_BODY()
 
 };
diff --git a/Engine/Source/Editor/UnrealEd/Classes/ActorFactories/ActorFactoryPawn.h b/Engine/Source/Editor/UnrealEd/Classes/ActorFactories/ActorFactoryPawn.h
--- a/Engine/Source/Editor/UnrealEd/Classes/ActorFactories/ActorFactoryPawn.h
+++ b/Engine/Source/Editor/UnrealEd/Classes/ActorFactories/ActorFactoryPawn.h
@@ -7,9 +7,7 @@
 UCLASS(MinimalAPI,config=Editor)
 class UActorFactoryPawn : public UActorFactoryEmptyActor
 {
-	GENERATED_BODY()
-public:
-	UNREALED_API UActorFactoryPawn(const FObjectInitializer& ObjectInitializer = FObjectInitializer::Get());
+	GENERATED_UCLASS_BODY()
 
 public:
 	virtual bool CanCreateActorFrom( const FAssetData& AssetData, FText& OutErrorMsg ) override;
diff --git a/Engine/Source/Editor/UnrealEd/Classes/ActorFactories/ActorFactoryPhysicsAsset.h b/Engine/Source/Editor/UnrealEd/Classes/ActorFactories/ActorFactoryPhysicsAsset.h
--- a/Engine/Source/Editor/UnrealEd/Classes/ActorFactories/ActorFactoryPhysicsAsset.h
+++ b/Engine/Source/Editor/UnrealEd/Classes/ActorFactories/ActorFactoryPhysicsAsset.h
@@ -8,9 +8,7 @@
 UCLASS(MinimalAPI, config=Editor, collapsecategories, hidecategories=Object)
 class UActorFactoryPhysicsAsset : public UActorFactory
 {
-	GENERATED_BODY()
-public:
-	UNREALED_API UActorFactoryPhysicsAsset(const FObjectInitializer& ObjectInitializer = FObjectInitializer::Get());
+	GENERATED_UCLASS_BODY()
 
 	// Begin UActorFactory Interface
 	virtual bool PreSpawnActor( UObject* Asset, FTransform& InOutLocation ) override;
diff --git a/Engine/Source/Editor/UnrealEd/Classes/ActorFactories/ActorFactoryPlaneReflectionCapture.h b/Engine/Source/Editor/UnrealEd/Classes/ActorFactories/ActorFactoryPlaneReflectionCapture.h
--- a/Engine/Source/Editor/UnrealEd/Classes/ActorFactories/ActorFactoryPlaneReflectionCapture.h
+++ b/Engine/Source/Editor/UnrealEd/Classes/ActorFactories/ActorFactoryPlaneReflectionCapture.h
@@ -7,9 +7,7 @@
 UCLASS(MinimalAPI, config=Editor, collapsecategories, hidecategories=Object)
 class UActorFactoryPlaneReflectionCapture : public UActorFactory
 {
-	GENERATED_BODY()
-public:
-	UNREALED_API UActorFactoryPlaneReflectionCapture(const FObjectInitializer& ObjectInitializer = FObjectInitializer::Get());
+	GENERATED_UCLASS_BODY()
 };
 
 
diff --git a/Engine/Source/Editor/UnrealEd/Classes/ActorFactories/ActorFactoryPlayerStart.h b/Engine/Source/Editor/UnrealEd/Classes/ActorFactories/ActorFactoryPlayerStart.h
--- a/Engine/Source/Editor/UnrealEd/Classes/ActorFactories/ActorFactoryPlayerStart.h
+++ b/Engine/Source/Editor/UnrealEd/Classes/ActorFactories/ActorFactoryPlayerStart.h
@@ -7,9 +7,7 @@
 UCLASS(MinimalAPI, config=Editor, collapsecategories, hidecategories=Object)
 class UActorFactoryPlayerStart : public UActorFactory
 {
-	GENERATED_BODY()
-public:
-	UNREALED_API UActorFactoryPlayerStart(const FObjectInitializer& ObjectInitializer = FObjectInitializer::Get());
+	GENERATED_UCLASS_BODY()
 
 };
 
diff --git a/Engine/Source/Editor/UnrealEd/Classes/ActorFactories/ActorFactoryPointLight.h b/Engine/Source/Editor/UnrealEd/Classes/ActorFactories/ActorFactoryPointLight.h
--- a/Engine/Source/Editor/UnrealEd/Classes/ActorFactories/ActorFactoryPointLight.h
+++ b/Engine/Source/Editor/UnrealEd/Classes/ActorFactories/ActorFactoryPointLight.h
@@ -7,9 +7,7 @@
 UCLASS(MinimalAPI, config=Editor, collapsecategories, hidecategories=Object)
 class UActorFactoryPointLight : public UActorFactory
 {
-	GENERATED_BODY()
-public:
-	UNREALED_API UActorFactoryPointLight(const FObjectInitializer& ObjectInitializer = FObjectInitializer::Get());
+	GENERATED_UCLASS_BODY()
 
 };
 
diff --git a/Engine/Source/Editor/UnrealEd/Classes/ActorFactories/ActorFactorySkeletalMesh.h b/Engine/Source/Editor/UnrealEd/Classes/ActorFactories/ActorFactorySkeletalMesh.h
--- a/Engine/Source/Editor/UnrealEd/Classes/ActorFactories/ActorFactorySkeletalMesh.h
+++ b/Engine/Source/Editor/UnrealEd/Classes/ActorFactories/ActorFactorySkeletalMesh.h
@@ -10,9 +10,7 @@ class USkeletalMesh;
 UCLASS(MinimalAPI, config=Editor)
 class UActorFactorySkeletalMesh : public UActorFactory
 {
-	GENERATED_BODY()
-public:
-	UNREALED_API UActorFactorySkeletalMesh(const FObjectInitializer& ObjectInitializer = FObjectInitializer::Get());
+	GENERATED_UCLASS_BODY()
 
 protected:
 	// Begin UActorFactory Interface
diff --git a/Engine/Source/Editor/UnrealEd/Classes/ActorFactories/ActorFactorySkyLight.h b/Engine/Source/Editor/UnrealEd/Classes/ActorFactories/ActorFactorySkyLight.h
--- a/Engine/Source/Editor/UnrealEd/Classes/ActorFactories/ActorFactorySkyLight.h
+++ b/Engine/Source/Editor/UnrealEd/Classes/ActorFactories/ActorFactorySkyLight.h
@@ -7,9 +7,7 @@
 UCLASS(MinimalAPI, config=Editor, collapsecategories, hidecategories=Object)
 class UActorFactorySkyLight : public UActorFactory
 {
-	GENERATED_BODY()
-public:
-	UNREALED_API UActorFactorySkyLight(const FObjectInitializer& ObjectInitializer = FObjectInitializer::Get());
+	GENERATED_UCLASS_BODY()
 
 };
 
diff --git a/Engine/Source/Editor/UnrealEd/Classes/ActorFactories/ActorFactorySphereVolume.h b/Engine/Source/Editor/UnrealEd/Classes/ActorFactories/ActorFactorySphereVolume.h
--- a/Engine/Source/Editor/UnrealEd/Classes/ActorFactories/ActorFactorySphereVolume.h
+++ b/Engine/Source/Editor/UnrealEd/Classes/ActorFactories/ActorFactorySphereVolume.h
@@ -7,9 +7,7 @@
 UCLASS(MinimalAPI, config=Editor, collapsecategories, hidecategories=Object)
 class UActorFactorySphereVolume : public UActorFactory
 {
-	GENERATED_BODY()
-public:
-	UNREALED_API UActorFactorySphereVolume(const FObjectInitializer& ObjectInitializer = FObjectInitializer::Get());
+	GENERATED_UCLASS_BODY()
 
 	// Begin UActorFactory Interface
 	virtual bool CanCreateActorFrom( const FAssetData& AssetData, FText& OutErrorMsg ) override;
diff --git a/Engine/Source/Editor/UnrealEd/Classes/ActorFactories/ActorFactorySpotLight.h b/Engine/Source/Editor/UnrealEd/Classes/ActorFactories/ActorFactorySpotLight.h
--- a/Engine/Source/Editor/UnrealEd/Classes/ActorFactories/ActorFactorySpotLight.h
+++ b/Engine/Source/Editor/UnrealEd/Classes/ActorFactories/ActorFactorySpotLight.h
@@ -7,9 +7,7 @@
 UCLASS(MinimalAPI, config=Editor, collapsecategories, hidecategories=Object)
 class UActorFactorySpotLight : public UActorFactory
 {
-	GENERATED_BODY()
-public:
-	UNREALED_API UActorFactorySpotLight(const FObjectInitializer& ObjectInitializer = FObjectInitializer::Get());
+	GENERATED_UCLASS_BODY()
 
 };
 
diff --git a/Engine/Source/Editor/UnrealEd/Classes/ActorFactories/ActorFactoryStaticMesh.h b/Engine/Source/Editor/UnrealEd/Classes/ActorFactories/ActorFactoryStaticMesh.h
--- a/Engine/Source/Editor/UnrealEd/Classes/ActorFactories/ActorFactoryStaticMesh.h
+++ b/Engine/Source/Editor/UnrealEd/Classes/ActorFactories/ActorFactoryStaticMesh.h
@@ -7,9 +7,7 @@
 UCLASS(MinimalAPI, config=Editor)
 class UActorFactoryStaticMesh : public UActorFactory
 {
-	GENERATED_BODY()
-public:
-	UNREALED_API UActorFactoryStaticMesh(const FObjectInitializer& ObjectInitializer = FObjectInitializer::Get());
+	GENERATED_UCLASS_BODY()
 
 	// Begin UActorFactory Interface
 	virtual bool CanCreateActorFrom( const FAssetData& AssetData, FText& OutErrorMsg ) override;
