Fix lightmass/editor failing to build if Swarm PDBs are not available.

diff --git a/Engine/Source/Editor/SwarmInterface/SwarmInterface.Build.cs b/Engine/Source/Editor/SwarmInterface/SwarmInterface.Build.cs
--- a/Engine/Source/Editor/SwarmInterface/SwarmInterface.Build.cs
+++ b/Engine/Source/Editor/SwarmInterface/SwarmInterface.Build.cs
@@ -1,6 +1,7 @@
 // Copyright 1998-2019 Epic Games, Inc. All Rights Reserved.
 
 using UnrealBuildTool;
+using System.IO;
 
 public class SwarmInterface : ModuleRules
 {
@@ -27,8 +28,13 @@ public class SwarmInterface : ModuleRules
 			PublicDependencyModuleNames.Add("Networking");
 		}
 
-		// Copy the AgentInterface DLLs to the same output directory as the editor DLL.
+		// Copy the AgentInterface DLL to the same output directory as the editor DLL.
 		RuntimeDependencies.Add("$(BinaryOutputDir)/AgentInterface.dll", "$(EngineDir)/Binaries/DotNET/AgentInterface.dll", StagedFileType.NonUFS);
-		RuntimeDependencies.Add("$(BinaryOutputDir)/AgentInterface.pdb", "$(EngineDir)/Binaries/DotNET/AgentInterface.pdb", StagedFileType.DebugNonUFS);
+
+		// Also copy the PDB, if it exists
+		if(File.Exists(Path.Combine(EngineDirectory, "Binaries", "DotNET", "AgentInterface.pdb")))
+		{
+			RuntimeDependencies.Add("$(BinaryOutputDir)/AgentInterface.pdb", "$(EngineDir)/Binaries/DotNET/AgentInterface.pdb", StagedFileType.DebugNonUFS);
+		}
 	}
 }
