Linux/Mac: fixed UnrealLightmass communication issues

diff --git a/Engine/Plugins/Messaging/UdpMessaging/UdpMessaging.uplugin b/Engine/Plugins/Messaging/UdpMessaging/UdpMessaging.uplugin
--- a/Engine/Plugins/Messaging/UdpMessaging/UdpMessaging.uplugin
+++ b/Engine/Plugins/Messaging/UdpMessaging/UdpMessaging.uplugin
@@ -24,7 +24,7 @@
 	[
 		{
 			"Name" : "UdpMessaging",
-			"Type" : "Runtime",
+			"Type" : "RuntimeAndProgram",
 			"LoadingPhase" : "PreDefault",
 			"BlacklistPlatforms" :
 			[
diff --git a/Engine/Source/Editor/SwarmInterface/Private/SwarmInterfaceLocal.cpp b/Engine/Source/Editor/SwarmInterface/Private/SwarmInterfaceLocal.cpp
--- a/Engine/Source/Editor/SwarmInterface/Private/SwarmInterfaceLocal.cpp
+++ b/Engine/Source/Editor/SwarmInterface/Private/SwarmInterfaceLocal.cpp
@@ -101,6 +101,8 @@ bool FSwarmInterface::Initialize(const TCHAR* SwarmInterfacePath)
 		GInstance = new FSwarmInterfaceLocalImpl();
 	}
 
+	FIPv4Endpoint::Initialize();
+
 	return true;
 }
 
@@ -231,6 +233,8 @@ int32 FSwarmInterfaceLocalImpl::SendMessage( const FMessage& Message )
 
 	while (bIsConnected && !Recepient.IsValid())
 	{
+		MessageEndpoint->Publish(new FSwarmPingMessage(), EMessageScope::Network);
+		
 		FTaskGraphInterface::Get().ProcessThreadUntilIdle(ENamedThreads::GameThread);
 		FPlatformProcess::Sleep(0.5f);
 
diff --git a/Engine/Source/Runtime/Networking/Public/Interfaces/IPv4/IPv4Endpoint.h b/Engine/Source/Runtime/Networking/Public/Interfaces/IPv4/IPv4Endpoint.h
--- a/Engine/Source/Runtime/Networking/Public/Interfaces/IPv4/IPv4Endpoint.h
+++ b/Engine/Source/Runtime/Networking/Public/Interfaces/IPv4/IPv4Endpoint.h
@@ -148,7 +148,7 @@ public:
 public:
 
 	/** Initializes the IP endpoint functionality. */
-	static void Initialize();
+	static NETWORKING_API void Initialize();
 
 	/**
 	 * Converts a string to an IPv4 endpoint.
