Use reliable messages for local swarm Mac/Linux messages

Udp window size is ~64, so when sending large batches of task completed messages quickly several would get silently dropped

With a batch of 512 messages, results would ~ be:

  first 68 sent successfully
  next 120 lost
  60 success
  145 lost
  70 success
  49 lost

This patch adds EMessageFlags::Reliable for all local swarm messages

diff --git a/Engine/Source/Editor/SwarmInterface/Private/SwarmInterfaceLocal.cpp b/Engine/Source/Editor/SwarmInterface/Private/SwarmInterfaceLocal.cpp
--- a/Engine/Source/Editor/SwarmInterface/Private/SwarmInterfaceLocal.cpp
+++ b/Engine/Source/Editor/SwarmInterface/Private/SwarmInterfaceLocal.cpp
@@ -81,10 +81,17 @@ private:
 	void* CallbackData;
 #if USE_LOCAL_SWARM_INTERFACE
 	TSharedPtr<FMessageEndpoint, ESPMode::ThreadSafe> MessageEndpoint;
-	FMessageAddress Recepient;
+	FMessageAddress Recipient;
 	bool bIsConnected;
 	bool bIsEditor;
 	FProcHandle LightmassProcHandle;
+
+	template<typename MessageType>
+	void MessageSendReliable(MessageType* Message, const FMessageAddress& RecipientAddr)
+	{
+		MessageEndpoint->Send(Message, MessageType::StaticStruct(), EMessageFlags::Reliable,
+			nullptr, TArrayBuilder<FMessageAddress>().Add(RecipientAddr), FTimespan::Zero(), FDateTime::MaxValue());
+	}
 #endif
 };
 
@@ -217,7 +224,7 @@ int32 FSwarmInterfaceLocalImpl::CloseConnection( void )
 		FPlatformProcess::TerminateProc(LightmassProcHandle, true);
 		FPlatformProcess::CloseProc(LightmassProcHandle);
 	}
-	Recepient = FMessageAddress();
+	Recipient = FMessageAddress();
 	MessageEndpoint.Reset();
 	bIsConnected = false;
 	CallbackFunc = NULL;
@@ -232,7 +239,7 @@ int32 FSwarmInterfaceLocalImpl::SendMessage( const FMessage& Message )
 	const double kMaxTimeToWaitSec = 60;
 	double TimeStartedWaiting = FPlatformTime::Seconds();
 
-	while (bIsConnected && !Recepient.IsValid())
+	while (bIsConnected && !Recipient.IsValid())
 	{
 		MessageEndpoint->Publish(FMessageEndpoint::MakeMessage<FSwarmPingMessage>(), EMessageScope::Network);
 		
@@ -250,17 +257,17 @@ int32 FSwarmInterfaceLocalImpl::SendMessage( const FMessage& Message )
 	if (Message.Type == MESSAGE_INFO)
 	{
 		FInfoMessage* InfoMessage = (FInfoMessage*)&Message;
-		MessageEndpoint->Send(FMessageEndpoint::MakeMessage<FSwarmInfoMessage>(InfoMessage->TextMessage), Recepient);
+		MessageSendReliable(FMessageEndpoint::MakeMessage<FSwarmInfoMessage>(InfoMessage->TextMessage), Recipient);
 	}
 	else if (Message.Type == MESSAGE_ALERT)
 	{
 		FAlertMessage* AlertMessage = (FAlertMessage*)&Message;
-		MessageEndpoint->Send(FMessageEndpoint::MakeMessage<FSwarmAlertMessage>(AlertMessage->JobGuid, AlertMessage->AlertLevel, AlertMessage->ObjectGuid, AlertMessage->TypeId, AlertMessage->TextMessage), Recepient);
+		MessageSendReliable(FMessageEndpoint::MakeMessage<FSwarmAlertMessage>(AlertMessage->JobGuid, AlertMessage->AlertLevel, AlertMessage->ObjectGuid, AlertMessage->TypeId, AlertMessage->TextMessage), Recipient);
 	}
 	else if (Message.Type == MESSAGE_TIMING)
 	{
 		FTimingMessage* TimingMessage = (FTimingMessage*)&Message;
-		MessageEndpoint->Send(FMessageEndpoint::MakeMessage<FSwarmTimingMessage>(TimingMessage->State, TimingMessage->ThreadNum), Recepient);
+		MessageSendReliable(FMessageEndpoint::MakeMessage<FSwarmTimingMessage>(TimingMessage->State, TimingMessage->ThreadNum), Recipient);
 	}
 	else if (Message.Type == MESSAGE_TASK_REQUEST)
 	{
@@ -275,18 +282,18 @@ int32 FSwarmInterfaceLocalImpl::SendMessage( const FMessage& Message )
 			FTaskSpecification TaskSpec = Tasks.Pop();
 			CallbackFunc( (FMessage*)&TaskSpec, CallbackData );
 		}
-		MessageEndpoint->Send(FMessageEndpoint::MakeMessage<FSwarmTaskRequestMessage>(), Recepient);
+		MessageSendReliable(FMessageEndpoint::MakeMessage<FSwarmTaskRequestMessage>(), Recipient);
 	}
 	else if (Message.Type == MESSAGE_TASK_REQUEST_RESPONSE)
 	{
 		FTaskRequestResponse* ResponseMessage = (FTaskRequestResponse*)&Message;
 		if (ResponseMessage->ResponseType == RESPONSE_TYPE_RELEASE)
 		{
-			MessageEndpoint->Send(FMessageEndpoint::MakeMessage<FSwarmTaskRequestReleaseMessage>(), Recepient);
+			MessageSendReliable(FMessageEndpoint::MakeMessage<FSwarmTaskRequestReleaseMessage>(), Recipient);
 		}
 		else if (ResponseMessage->ResponseType == RESPONSE_TYPE_RESERVATION)
 		{
-			MessageEndpoint->Send(FMessageEndpoint::MakeMessage<FSwarmTaskRequestReservationMessage>(), Recepient);
+			MessageSendReliable(FMessageEndpoint::MakeMessage<FSwarmTaskRequestReservationMessage>(), Recipient);
 		}
 		else if (ResponseMessage->ResponseType == RESPONSE_TYPE_SPECIFICATION)
 		{
@@ -296,22 +303,22 @@ int32 FSwarmInterfaceLocalImpl::SendMessage( const FMessage& Message )
 			{
 				Dependencies.Add(SpecificationMessage->Dependencies[Index]);
 			}
-			MessageEndpoint->Send(FMessageEndpoint::MakeMessage<FSwarmTaskRequestSpecificationMessage>(SpecificationMessage->TaskGuid, SpecificationMessage->Parameters, SpecificationMessage->Flags, Dependencies), Recepient);
+			MessageSendReliable(FMessageEndpoint::MakeMessage<FSwarmTaskRequestSpecificationMessage>(SpecificationMessage->TaskGuid, SpecificationMessage->Parameters, SpecificationMessage->Flags, Dependencies), Recipient);
 		}
 	}
 	else if (Message.Type == MESSAGE_JOB_STATE)
 	{
 		FJobState* StateMessage = (FJobState*)&Message;
-		MessageEndpoint->Send(FMessageEndpoint::MakeMessage<FSwarmJobStateMessage>(StateMessage->JobGuid, StateMessage->JobState, StateMessage->JobMessage, StateMessage->JobExitCode, StateMessage->JobRunningTime), Recepient);
+		MessageSendReliable(FMessageEndpoint::MakeMessage<FSwarmJobStateMessage>(StateMessage->JobGuid, StateMessage->JobState, StateMessage->JobMessage, StateMessage->JobExitCode, StateMessage->JobRunningTime), Recipient);
 	}
 	else if (Message.Type == MESSAGE_TASK_STATE)
 	{
 		FTaskState* StateMessage = (FTaskState*)&Message;
-		MessageEndpoint->Send(FMessageEndpoint::MakeMessage<FSwarmTaskStateMessage>(StateMessage->TaskGuid, StateMessage->TaskState, StateMessage->TaskMessage, StateMessage->TaskExitCode, StateMessage->TaskRunningTime), Recepient);
+		MessageSendReliable(FMessageEndpoint::MakeMessage<FSwarmTaskStateMessage>(StateMessage->TaskGuid, StateMessage->TaskState, StateMessage->TaskMessage, StateMessage->TaskExitCode, StateMessage->TaskRunningTime), Recipient);
 	}
 	else if (Message.Type == MESSAGE_QUIT)
 	{
-		MessageEndpoint->Send(FMessageEndpoint::MakeMessage<FSwarmQuitMessage>(), Recepient);
+		MessageSendReliable(FMessageEndpoint::MakeMessage<FSwarmQuitMessage>(), Recipient);
 	}
 #endif
 	return 0;
@@ -320,14 +327,14 @@ int32 FSwarmInterfaceLocalImpl::SendMessage( const FMessage& Message )
 #if USE_LOCAL_SWARM_INTERFACE
 void FSwarmInterfaceLocalImpl::HandlePingMessage( const FSwarmPingMessage& Message, const TSharedRef<IMessageContext, ESPMode::ThreadSafe>& Context )
 {
-	MessageEndpoint->Send(FMessageEndpoint::MakeMessage<FSwarmPongMessage>(bIsEditor, FPlatformProcess::ComputerName()), Context->GetSender());
+	MessageSendReliable(FMessageEndpoint::MakeMessage<FSwarmPongMessage>(bIsEditor, FPlatformProcess::ComputerName()), Context->GetSender());
 }
 
 void FSwarmInterfaceLocalImpl::HandlePongMessage( const FSwarmPongMessage& Message, const TSharedRef<IMessageContext, ESPMode::ThreadSafe>& Context )
 {
-	if (!Recepient.IsValid() && Message.bIsEditor != bIsEditor && Message.ComputerName == FPlatformProcess::ComputerName())
+	if (!Recipient.IsValid() && Message.bIsEditor != bIsEditor && Message.ComputerName == FPlatformProcess::ComputerName())
 	{
-		Recepient = Context->GetSender();
+		Recipient = Context->GetSender();
 	}
 }
 
