UnrealBuildTool: Only include NetFxSDK when required by a module

#rnx

diff --git a/Engine/Source/Editor/SwarmInterface/SwarmInterface.Build.cs b/Engine/Source/Editor/SwarmInterface/SwarmInterface.Build.cs
--- a/Engine/Source/Editor/SwarmInterface/SwarmInterface.Build.cs
+++ b/Engine/Source/Editor/SwarmInterface/SwarmInterface.Build.cs
@@ -27,6 +27,19 @@ public class SwarmInterface : ModuleRules
 			PublicDependencyModuleNames.Add("Sockets");
 			PublicDependencyModuleNames.Add("Networking");
 		}
+		else if (Target.Platform == UnrealTargetPlatform.Win64)
+		{
+			string NetFxSdkDir = Target.WindowsPlatform.NetFxSdkDir;
+			if (NetFxSdkDir == null)
+			{
+				throw new BuildException("Could not find NetFxSDK install dir; this will prevent SwarmInterface from installing.  Install a version of .NET Framework SDK at 4.6.0 or higher.");
+			}
+
+			string ArchFolder = WindowsExports.GetArchitectureSubpath(Target.WindowsPlatform.Architecture);
+
+			PrivateIncludePaths.Add(Path.Combine(NetFxSdkDir, "include", "um"));
+			PublicSystemLibraryPaths.Add(Path.Combine(NetFxSdkDir, "lib", "um", ArchFolder));
+		}
 
 		// Copy the AgentInterface DLL to the same output directory as the editor DLL.
 		RuntimeDependencies.Add("$(BinaryOutputDir)/AgentInterface.dll", "$(EngineDir)/Binaries/DotNET/AgentInterface.dll", StagedFileType.NonUFS);
diff --git a/Engine/Source/Programs/UnrealBuildTool/Platform/Windows/UEBuildWindows.cs b/Engine/Source/Programs/UnrealBuildTool/Platform/Windows/UEBuildWindows.cs
--- a/Engine/Source/Programs/UnrealBuildTool/Platform/Windows/UEBuildWindows.cs
+++ b/Engine/Source/Programs/UnrealBuildTool/Platform/Windows/UEBuildWindows.cs
@@ -411,7 +411,25 @@
 			get { return (Environment == null)? null : Environment.WindowsSdkDir.FullName; }
 		}
 
+
 		/// <summary>
+		/// Directory containing the NETFXSDK
+		/// </summary>
+		[SuppressMessage("Interoperability", "CA1416:Validate platform compatibility", Justification = "Manually checked")]
+		public string? NetFxSdkDir
+		{
+			get
+			{
+				DirectoryReference? NetFxSdkDir;
+				if (RuntimePlatform.IsWindows && MicrosoftPlatformSDK.TryGetNetFxSdkInstallDir(out NetFxSdkDir))
+				{
+					return NetFxSdkDir.FullName;
+				}
+				return null;
+			}
+		}
+
+		/// <summary>
 		/// Directory containing the DIA SDK
 		/// </summary>
 		public string? DiaSdkDir
@@ -684,6 +702,11 @@
 			get { return Inner.WindowsSdkDir; }
 		}
 
+		public string? NetFxSdkDir
+		{
+			get { return Inner.NetFxSdkDir; }
+		}
+
 		public string? DiaSdkDir
 		{
 			get { return Inner.DiaSdkDir; }
diff --git a/Engine/Source/Programs/UnrealBuildTool/Platform/Windows/VCEnvironment.cs b/Engine/Source/Programs/UnrealBuildTool/Platform/Windows/VCEnvironment.cs
--- a/Engine/Source/Programs/UnrealBuildTool/Platform/Windows/VCEnvironment.cs
+++ b/Engine/Source/Programs/UnrealBuildTool/Platform/Windows/VCEnvironment.cs
@@ -327,18 +327,6 @@
 				LibraryPaths.Add(DirectoryReference.Combine(LibraryRootDir, "ucrt", ArchFolder));
 			}
 
-			// Add the NETFXSDK include path. We need this for SwarmInterface.
-			DirectoryReference? NetFxSdkDir;
-			if(MicrosoftPlatformSDK.TryGetNetFxSdkInstallDir(out NetFxSdkDir))
-			{
-				IncludePaths.Add(DirectoryReference.Combine(NetFxSdkDir, "include", "um"));
-				LibraryPaths.Add(DirectoryReference.Combine(NetFxSdkDir, "lib", "um", ArchFolder));
-			}
-			else
-			{
-				throw new BuildException("Could not find NetFxSDK install dir; this will prevent SwarmInterface from installing.  Install a version of .NET Framework SDK at 4.6.0 or higher.");
-			}
-
 			// Add the Windows SDK paths
 			if (WindowsSdkVersion >= new VersionNumber(10))
 			{
