Added so ExtraRootPath is set for all modules that has special extra external folders referenced

diff --git a/Engine/Source/Editor/SwarmInterface/SwarmInterface.Build.cs b/Engine/Source/Editor/SwarmInterface/SwarmInterface.Build.cs
--- a/Engine/Source/Editor/SwarmInterface/SwarmInterface.Build.cs
+++ b/Engine/Source/Editor/SwarmInterface/SwarmInterface.Build.cs
@@ -29,6 +29,8 @@ public class SwarmInterface : ModuleRules
 				throw new BuildException("Could not find NetFxSDK install dir; this will prevent SwarmInterface from installing.  Install a version of .NET Framework SDK at 4.6.0 or higher.");
 			}
 
+			ExtraRootPath = NetFxSdkDir;
+
 			string ArchFolder = Target.WindowsPlatform.Architecture.WindowsLibDir;
 
 			PrivateIncludePaths.Add(Path.Combine(NetFxSdkDir, "include", "um"));
