Disable Swarm for Windows Arm64

#rnx

diff --git a/Engine/Source/Editor/SwarmInterface/Private/SwarmInterface.cpp b/Engine/Source/Editor/SwarmInterface/Private/SwarmInterface.cpp
--- a/Engine/Source/Editor/SwarmInterface/Private/SwarmInterface.cpp
+++ b/Engine/Source/Editor/SwarmInterface/Private/SwarmInterface.cpp
@@ -2,7 +2,7 @@
 
 #include "SwarmInterface.h"
 
-#if PLATFORM_WINDOWS
+#if PLATFORM_WINDOWS && !PLATFORM_CPU_ARM_FAMILY
 
 // To avoid compile problems with C++/CLI in VS11
 #pragma warning(disable:4538)
@@ -31,7 +31,7 @@ EXTERN_C IMAGE_DOS_HEADER __ImageBase;
 
 #include "Windows/HideWindowsPlatformTypes.h"
 
-#endif // PLATFORM_WINDOWS
+#endif // PLATFORM_WINDOWS && !PLATFORM_CPU_ARM_FAMILY
 
 
 namespace NSwarm
@@ -527,7 +527,7 @@ int32 FSwarmInterfaceImpl::Log( TVerbosityLevel Verbosity, TLogColour TextColour
 
 bool FSwarmInterfaceImpl::InitSwarmInterfaceManaged(const TCHAR* SwarmInterfaceDLLPath)
 {
-#if PLATFORM_WINDOWS
+#if PLATFORM_WINDOWS && !PLATFORM_CPU_ARM_FAMILY
 	ICLRMetaHost *MetaHost = NULL;
 	ICLRRuntimeHost* RuntimeHost = NULL;
 
@@ -580,7 +580,7 @@ bool FSwarmInterfaceImpl::InitSwarmInterfaceManaged(const TCHAR* SwarmInterfaceD
 		ensureMsgf(false, TEXT("Error creating Swarm instance bridge."));
 		return false;
 	}
-#endif // PLATFORM_WINDOWS
+#endif // PLATFORM_WINDOWS && !PLATFORM_CPU_ARM_FAMILY
 
 	return true;
 }
diff --git a/Engine/Source/Editor/SwarmInterface/SwarmInterface.Build.cs b/Engine/Source/Editor/SwarmInterface/SwarmInterface.Build.cs
--- a/Engine/Source/Editor/SwarmInterface/SwarmInterface.Build.cs
+++ b/Engine/Source/Editor/SwarmInterface/SwarmInterface.Build.cs
@@ -14,14 +14,16 @@ public class SwarmInterface : ModuleRules
 			}
 		);
 
+		bool bHasSwarm = false;
 		if (Target.Platform == UnrealTargetPlatform.Mac || Target.Platform == UnrealTargetPlatform.Linux)
 		{
 			PrivateIncludePathModuleNames.Add("MessagingCommon");
 			// the modules below are only needed for the UMB usability check
 			PublicDependencyModuleNames.Add("Sockets");
 			PublicDependencyModuleNames.Add("Networking");
+			bHasSwarm = true;
 		}
-		else if (Target.Platform == UnrealTargetPlatform.Win64)
+		else if (Target.Platform == UnrealTargetPlatform.Win64 && Target.Architecture != UnrealArch.Arm64)
 		{
 			string NetFxSdkDir = Target.WindowsPlatform.NetFxSdkDir;
 			if (NetFxSdkDir == null)
@@ -35,15 +37,19 @@ public class SwarmInterface : ModuleRules
 
 			PrivateIncludePaths.Add(Path.Combine(NetFxSdkDir, "include", "um"));
 			PublicSystemLibraryPaths.Add(Path.Combine(NetFxSdkDir, "lib", "um", ArchFolder));
+			bHasSwarm = true;
 		}
 
-		// Copy the AgentInterface DLL to the same output directory as the editor DLL.
-		RuntimeDependencies.Add("$(BinaryOutputDir)/AgentInterface.dll", "$(EngineDir)/Binaries/DotNET/AgentInterface.dll", StagedFileType.NonUFS);
-
-		// Also copy the PDB, if it exists
-		if(File.Exists(Path.Combine(EngineDirectory, "Binaries", "DotNET", "AgentInterface.pdb")))
+		if (bHasSwarm)
 		{
-			RuntimeDependencies.Add("$(BinaryOutputDir)/AgentInterface.pdb", "$(EngineDir)/Binaries/DotNET/AgentInterface.pdb", StagedFileType.DebugNonUFS);
+			// Copy the AgentInterface DLL to the same output directory as the editor DLL.
+			RuntimeDependencies.Add("$(BinaryOutputDir)/AgentInterface.dll", "$(EngineDir)/Binaries/DotNET/AgentInterface.dll", StagedFileType.NonUFS);
+
+			// Also copy the PDB, if it exists
+			if(File.Exists(Path.Combine(EngineDirectory, "Binaries", "DotNET", "AgentInterface.pdb")))
+			{
+				RuntimeDependencies.Add("$(BinaryOutputDir)/AgentInterface.pdb", "$(EngineDir)/Binaries/DotNET/AgentInterface.pdb", StagedFileType.DebugNonUFS);
+			}
 		}
 	}
 }
diff --git a/Engine/Source/Editor/UnrealEd/Private/Lightmass/Lightmass.h b/Engine/Source/Editor/UnrealEd/Private/Lightmass/Lightmass.h
--- a/Engine/Source/Editor/UnrealEd/Private/Lightmass/Lightmass.h
+++ b/Engine/Source/Editor/UnrealEd/Private/Lightmass/Lightmass.h
@@ -15,12 +15,12 @@
 #include "Lightmass/LightmassImportanceVolume.h"
 #include "Components/LightmassPortalComponent.h"
 #include "Components/SkyAtmosphereComponent.h"
-#if PLATFORM_WINDOWS
+#if PLATFORM_WINDOWS && !PLATFORM_CPU_ARM_FAMILY
 #include "Windows/AllowWindowsPlatformTypes.h"
 #endif
 	#include "SwarmInterface.h"
 	#include "Lightmass/LightmassRender.h"
-#if PLATFORM_WINDOWS
+#if PLATFORM_WINDOWS && !PLATFORM_CPU_ARM_FAMILY
 #include "Windows/HideWindowsPlatformTypes.h"
 #endif
 
